Added KPI Hints System
Fixed mismatched LOCTEXT #define

diff --git a/Engine/Plugins/Experimental/EditorPerformance/Source/EditorPerformance/Private/EditorPerformanceModule.cpp b/Engine/Plugins/Experimental/EditorPerformance/Source/EditorPerformance/Private/EditorPerformanceModule.cpp
--- a/Engine/Plugins/Experimental/EditorPerformance/Source/EditorPerformance/Private/EditorPerformanceModule.cpp
+++ b/Engine/Plugins/Experimental/EditorPerformance/Source/EditorPerformance/Private/EditorPerformanceModule.cpp
@@ -21,8 +21,6 @@
 #include "StudioTelemetry.h"
 #include "HAL/PlatformFileManager.h"
 
-UE_DISABLE_OPTIMIZATION_SHIP
-
 #define LOCTEXT_NAMESPACE "EditorPerformance"
  
 IMPLEMENT_MODULE(FEditorPerformanceModule, EditorPerformance );
@@ -186,6 +184,9 @@ void FEditorPerformanceModule::InitializeKPIs()
 	// Load the KPI profiles
 	KPIRegistry.LoadKPIProfiles(TEXT("EditorPerformance.Profile"), GEditorIni);
 
+	// Load the KPI hints
+	KPIRegistry.LoadKPIHints(TEXT("EditorPerformance.Hints"), GEditorIni);
+
 	// Apply any non map specific profiles
 	for (FKPIProfiles::TConstIterator It(KPIRegistry.GetKPIProfiles()); It; ++It)
 	{
@@ -384,15 +385,17 @@ bool FEditorPerformanceModule::RecordTelemetryEvent(const FKPIValue& KPIValue)
 	// Record a new telemetry event for this KPI
 	if (FStudioTelemetry::IsAvailable())
 	{
-		const int SchemaVersion = 1;
+		const int SchemaVersion = 2;
 		TArray<FAnalyticsEventAttribute> Attributes;
 
 		Attributes.Emplace(TEXT("SchemaVersion"), SchemaVersion);
 		Attributes.Emplace(TEXT("MapName"), EditorMapName);
+		Attributes.Emplace(TEXT("DDC_IsHotLocalCache"), IsHotLocalCacheCase());
 		Attributes.Emplace(TEXT("KPI_Name"), *KPIValue.Name.ToString());
 		Attributes.Emplace(TEXT("KPI_Category"), *KPIValue.Category.ToString());
 		Attributes.Emplace(TEXT("KPI_CurrentValue"), KPIValue.CurrentValue);
 		Attributes.Emplace(TEXT("KPI_ThresholdValue"), KPIValue.ThresholdValue);
+		Attributes.Emplace(TEXT("KPI_DisplayType"), *FKPIValue::GetDisplayTypeAsString(KPIValue.DisplayType));
 		Attributes.Emplace(TEXT("KPI_Profile"), KPIProfileName);
 		
 		FStudioTelemetry::Get().RecordEvent(TEXT("Editor.PerformanceWarning"), Attributes);
@@ -407,9 +410,9 @@ void FEditorPerformanceModule::TerminateKPIs()
 {
 }
 
-const FKPIValues& FEditorPerformanceModule::GetKPIValues() const
+const FKPIRegistry& FEditorPerformanceModule::GetKPIRegistry() const
 {
-	return KPIRegistry.GetKPIValues();
+	return KPIRegistry;
 }
 
 const FString& FEditorPerformanceModule::GetKPIProfileName() const
@@ -418,5 +421,3 @@ const FString& FEditorPerformanceModule::GetKPIProfileName() const
 }
 
 #undef LOCTEXT_NAMESPACE
-
-UE_ENABLE_OPTIMIZATION_SHIP
diff --git a/Engine/Plugins/Experimental/EditorPerformance/Source/EditorPerformance/Private/KPIValue.cpp b/Engine/Plugins/Experimental/EditorPerformance/Source/EditorPerformance/Private/KPIValue.cpp
--- a/Engine/Plugins/Experimental/EditorPerformance/Source/EditorPerformance/Private/KPIValue.cpp
+++ b/Engine/Plugins/Experimental/EditorPerformance/Source/EditorPerformance/Private/KPIValue.cpp
@@ -3,8 +3,6 @@
 #include "KPIValue.h"
 #include "Misc/ConfigCacheIni.h"
 
-UE_DISABLE_OPTIMIZATION_SHIP
-
 FKPIValue::EState FKPIValue::GetState() const
 {
 	return State;
@@ -106,6 +104,68 @@ FString FKPIValue::GetComparisonAsPrettyString(FKPIValue::ECompare Compare)
 	}
 }
 
+FString FKPIValue::GetDisplayTypeAsString(FKPIValue::EDisplayType DisplayType)
+{
+	switch (DisplayType)
+	{
+		default:
+		case EDisplayType::Decimal:
+		{
+			return TEXT("Decimal");
+			break;
+		}
+
+		case EDisplayType::Minutes:
+		{
+			return TEXT("Minutes");
+			break;
+		}
+
+		case EDisplayType::Seconds:
+		{
+			return TEXT("Seconds");
+			break;
+		}
+
+		case EDisplayType::Milliseconds:
+		{
+			return TEXT("Milliseconds");
+			break;
+		}
+
+		case EDisplayType::Bytes:
+		{
+			return TEXT("Bytes");
+			break;
+		}
+
+		case EDisplayType::MegaBytes:
+		{
+			return TEXT("MegaBytes");
+			break;
+		}
+
+		case EDisplayType::GigaBytes:
+		{
+			return TEXT("GigaBytes");
+			break;
+		}
+
+		case EDisplayType::MegaBitsPerSecond:
+		{
+			return TEXT("MegaBitsPerSecond");
+			break;
+		}
+
+		case EDisplayType::Percent:
+		{
+			return TEXT("Percent");
+			break;
+		}
+	}
+}
+
+
 FString	FKPIValue::GetValueAsString( float Value, FKPIValue::EDisplayType DisplayType)
 {
 	switch (DisplayType)
@@ -244,6 +304,19 @@ bool FKPIRegistry::GetKPIValue(const FName Name, FKPIValue& Result) const
 	return false;
 }
 
+bool FKPIRegistry::GetKPIHint(const FName Name, FKPIHint& Result) const
+{
+	const FKPIHint* Hint = Hints.Find(Name);
+
+	if (Hint != nullptr)
+	{
+		Result = *Hint;
+		return true;
+	}
+
+	return false;
+}
+
 const TMap<FName, FKPIValue>& FKPIRegistry::GetKPIValues() const
 {
 	return Values;
@@ -254,6 +327,34 @@ const FKPIProfiles& FKPIRegistry::GetKPIProfiles() const
 	return Profiles;
 }
 
+void FKPIRegistry::LoadKPIHints(const FString& HintSectionName, const FString& FileName)
+{
+	TArray<FString> SectionNames;
+
+	if (GConfig->GetSectionNames(FileName, SectionNames))
+	{
+		for (const FString& SectionName : SectionNames)
+		{
+			if (SectionName.Find(HintSectionName) != INDEX_NONE)
+			{
+				FKPIHint Hint;
+				Hint.URL = TEXT("https://docs.unrealengine.com/5.0/en-US/");
+
+				for (FKPIValues::TConstIterator It(GetKPIValues()); It; ++It)
+				{
+					const FKPIValue& KPIValue = It->Value;
+					FString KPIName = FString::Printf(TEXT("%s_%s"), *KPIValue.Category.ToString(), *KPIValue.Name.ToString()).Replace(TEXT(" "), TEXT("_"));
+
+					if (GConfig->GetString(*SectionName, *KPIName, Hint.Message, FileName))
+					{
+						Hints.Emplace(It->Key, Hint);
+					}
+				}
+			}
+		}
+	}
+}
+
 void FKPIRegistry::LoadKPIProfiles(const FString& ProfileSectionName, const FString& FileName)
 {
 	TArray<FString> SectionNames;
@@ -307,7 +408,3 @@ bool FKPIRegistry::ApplyKPIProfile(const FKPIProfile& KPIProfile)
 
 	return Result;
 }
-
-#undef LOCTEXT_NAMESPACE
-
-UE_ENABLE_OPTIMIZATION_SHIP
diff --git a/Engine/Plugins/Experimental/EditorPerformance/Source/EditorPerformance/Private/SEditorPerformanceDialogs.cpp b/Engine/Plugins/Experimental/EditorPerformance/Source/EditorPerformance/Private/SEditorPerformanceDialogs.cpp
--- a/Engine/Plugins/Experimental/EditorPerformance/Source/EditorPerformance/Private/SEditorPerformanceDialogs.cpp
+++ b/Engine/Plugins/Experimental/EditorPerformance/Source/EditorPerformance/Private/SEditorPerformanceDialogs.cpp
@@ -12,14 +12,13 @@
 #include "Widgets/SToolTip.h"
 #include "Widgets/Images/SImage.h"
 #include "Widgets/Layout/SGridPanel.h"
+#include "Widgets/Layout/SScrollBox.h"
 #include "Widgets/Text/STextBlock.h"
 #include "Widgets/Input/SComboBox.h"
 #include "Widgets/Input/SCheckBox.h"
 #include "EditorPerformanceModule.h"
 #include "Editor/EditorPerformanceSettings.h"
 
-UE_DISABLE_OPTIMIZATION_SHIP
-
 #define LOCTEXT_NAMESPACE "EditorPerformance"
 
 void SEditorPerformanceReportDialog::Construct(const FArguments& InArgs)
@@ -73,9 +72,9 @@ void SEditorPerformanceReportDialog::Construct(const FArguments& InArgs)
 		+ SVerticalBox::Slot()
 		.AutoHeight()
 		.Padding(0, 5, 0, 0)
-		.Expose(InformationGridSlot)
+		.Expose(HintGridSlot)
 		[
-			GetInformationGridPanel()
+			GetHintGridPanel()
 		]
 	];
 
@@ -94,9 +93,9 @@ EActiveTimerReturnType SEditorPerformanceReportDialog::UpdateGridPanels(double I
 		GetKPIGridPanel()
 	];
 
-	(*InformationGridSlot)
+	(*HintGridSlot)
 	[
-		GetInformationGridPanel()
+		GetHintGridPanel()
 	];
 
 	SlatePrepass(GetPrepassLayoutScaleMultiplier());
@@ -228,7 +227,12 @@ TSharedRef<SWidget> SEditorPerformanceReportDialog::GetSettingsGridPanel()
 	return Panel;
 }
 
-TSharedRef<SWidget> SEditorPerformanceReportDialog::GetInformationGridPanel()
+void SEditorPerformanceReportDialog::OpenURL(const FString& URL) const
+{
+	FPlatformProcess::LaunchURL(*URL, nullptr, nullptr);
+}
+
+TSharedRef<SWidget> SEditorPerformanceReportDialog::GetHintGridPanel()
 {
 	TSharedRef<SGridPanel> Panel =
 		SNew(SGridPanel);
@@ -243,33 +247,70 @@ TSharedRef<SWidget> SEditorPerformanceReportDialog::GetInformationGridPanel()
 	const FMargin DefaultMargin(0.0f, RowMargin, ColumnMargin, RowMargin);
 	const FMargin DefaultMarginFirstColumn(ColumnMargin, RowMargin);
 
+	FEditorPerformanceModule& EditorPerfModule = FModuleManager::LoadModuleChecked<FEditorPerformanceModule>("EditorPerformance");
+
 	int32 Row = 0;
 
-	Panel->AddSlot(0, Row)
-		.HAlign(HAlign_Left)
-		[
-			SNew(STextBlock)
-			.Margin(TitleMarginFirstColumn)
-			.ColorAndOpacity(TitleColor)
-			.Font(TitleFont)
-			.Justification(ETextJustify::Left)
-			.Text(LOCTEXT("InformationTitle", "Information"))
-		];
+	FKPIHint KPIHint;
 
-	Row++;
+	for (FKPIValues::TConstIterator It(EditorPerfModule.GetKPIRegistry().GetKPIValues()); It; ++It)
+	{
+		const FKPIValue& KPIValue = It->Value;
 
-	Panel->AddSlot(0, Row)
-		.HAlign(HAlign_Left)
-		[
-			SNew(STextBlock)
-			.Margin(DefaultMarginFirstColumn)
-			.ColorAndOpacity(EStyleColor::Foreground)
-			.Font(TitleFont)
-			.Justification(ETextJustify::Left)
-			.Text(LOCTEXT("DynamicText", "TODO : Some information about about what to do with this wanring"))
-		];
+		if (KPIValue.State==FKPIValue::Bad && EditorPerfModule.GetKPIRegistry().GetKPIHint(KPIValue.Name, KPIHint))
+		{
+			Panel->AddSlot(0, Row)
+				.HAlign(HAlign_Left)
+				[
+					SNew(STextBlock)
+					.Margin(TitleMarginFirstColumn)
+					.ColorAndOpacity(TitleColor)
+					.Font(TitleFont)
+					.Justification(ETextJustify::Left)
+					.Text(LOCTEXT("HintTitle", "Hint"))
+				];
 
-	Row++;
+			Row++;
+
+			Panel->AddSlot(0, Row)
+				.HAlign(HAlign_Left)
+				[
+					SNew(STextBlock)
+					.Margin(TitleMarginFirstColumn)
+					.ColorAndOpacity(TitleColor)
+					.Font(TitleFont)
+					.Text(FText::FromString(*FString::Printf(TEXT("%s %s"), *KPIValue.Category.ToString(), *KPIValue.Name.ToString())))
+				];
+
+			Row++;
+
+			Panel->AddSlot(0, Row)
+				.HAlign(HAlign_Left)
+				[
+					SNew(STextBlock)
+					.AutoWrapText(true)
+					.Margin(DefaultMarginFirstColumn)
+					.ColorAndOpacity(EStyleColor::Foreground)
+					.Justification(ETextJustify::Left)
+					.Text(FText::FromString(*KPIHint.Message))
+				];
+
+			Row++;
+
+			Panel->AddSlot(0, Row)
+				.HAlign(HAlign_Left)
+				[
+					SNew(STextBlock)
+					.AutoWrapText(true)
+					.Margin(DefaultMarginFirstColumn)
+					.ColorAndOpacity(EStyleColor::AccentBlue)
+					.Justification(ETextJustify::Left)
+					.Text(FText::FromString(*KPIHint.URL))
+				];
+
+			break;
+		}
+	}
 
 	return Panel;
 }
@@ -379,7 +420,7 @@ TSharedRef<SWidget> SEditorPerformanceReportDialog::GetKPIGridPanel()
 
 	TMap<FName, TArray<FKPIValue>> SortedKPIValues;
 
-	for (FKPIValues::TConstIterator It(EditorPerfModule.GetKPIValues()); It; ++It)
+	for (FKPIValues::TConstIterator It(EditorPerfModule.GetKPIRegistry().GetKPIValues()); It; ++It)
 	{
 		const FKPIValue& KPIValue = It->Value;
 
@@ -541,5 +582,3 @@ TSharedRef<SWidget> SEditorPerformanceReportDialog::GetKPIGridPanel()
 }
 
 #undef LOCTEXT_NAMESPACE
-
-UE_ENABLE_OPTIMIZATION_SHIP
diff --git a/Engine/Plugins/Experimental/EditorPerformance/Source/EditorPerformance/Private/SEditorPerformanceDialogs.h b/Engine/Plugins/Experimental/EditorPerformance/Source/EditorPerformance/Private/SEditorPerformanceDialogs.h
--- a/Engine/Plugins/Experimental/EditorPerformance/Source/EditorPerformance/Private/SEditorPerformanceDialogs.h
+++ b/Engine/Plugins/Experimental/EditorPerformance/Source/EditorPerformance/Private/SEditorPerformanceDialogs.h
@@ -20,14 +20,15 @@ class SEditorPerformanceReportDialog : public SCompoundWidget
 private:
 	
 	EActiveTimerReturnType UpdateGridPanels(double InCurrentTime, float InDeltaTime);
+	void OpenURL(const FString& URL) const;
 
 	TSharedRef<SWidget> GetKPIGridPanel();
 	TSharedRef<SWidget> GetSettingsGridPanel();
-	TSharedRef<SWidget> GetInformationGridPanel();
+	TSharedRef<SWidget> GetHintGridPanel();
 
 	SVerticalBox::FSlot* SettingsGridSlot = nullptr;
 	SVerticalBox::FSlot* KPIGridSlot = nullptr;
-	SVerticalBox::FSlot* InformationGridSlot = nullptr;
+	SVerticalBox::FSlot* HintGridSlot = nullptr;
 
 	const TArray<FName> NotifcationOptions = { FName("Notify"), FName("Ignore") };
 	const TArray<FName> WarningFilterOptions = { FName("Show All"), FName("Show Warnings") };
diff --git a/Engine/Plugins/Experimental/EditorPerformance/Source/EditorPerformance/Private/SEditorPerformanceStatusBar.cpp b/Engine/Plugins/Experimental/EditorPerformance/Source/EditorPerformance/Private/SEditorPerformanceStatusBar.cpp
--- a/Engine/Plugins/Experimental/EditorPerformance/Source/EditorPerformance/Private/SEditorPerformanceStatusBar.cpp
+++ b/Engine/Plugins/Experimental/EditorPerformance/Source/EditorPerformance/Private/SEditorPerformanceStatusBar.cpp
@@ -46,8 +46,6 @@
 #include "Settings/EditorProjectSettings.h"
 #include "Editor/EditorPerformanceSettings.h"
 
-UE_DISABLE_OPTIMIZATION_SHIP
-
 class SWidget;
 struct FSlateBrush;
 
@@ -127,39 +125,39 @@ void SEditorPerformanceStatusBarWidget::Construct(const FArguments& InArgs)
 		[
 			SNew(SComboButton)
 			.ContentPadding(FMargin(6.0f, 0.0f))
-		.MenuPlacement(MenuPlacement_AboveAnchor)
-		.ComboButtonStyle(&FAppStyle::Get().GetWidgetStyle<FComboButtonStyle>("SimpleComboButton"))
-		.ButtonContent()
-		[
+			.MenuPlacement(MenuPlacement_AboveAnchor)
+			.ComboButtonStyle(&FAppStyle::Get().GetWidgetStyle<FComboButtonStyle>("SimpleComboButton"))
+			.ButtonContent()
+			[
+
+				SNew(SHorizontalBox)
+				+ SHorizontalBox::Slot()
+				.AutoWidth()
+				.VAlign(VAlign_Center)
+				.Padding(0, 0, 3, 0)
+				[
+					SNew(SOverlay)
+					+ SOverlay::Slot()
+					.HAlign(HAlign_Center)
+					.VAlign(VAlign_Top)
+				[
+					SNew(SImage)
+					.ColorAndOpacity_Lambda([this] { return GetStatusIconColor(); })
+					.Image_Lambda([this] { return GetStatusIcon(); })
+					.ToolTipText_Lambda([this] { return GetStatusToolTipText(); })
+				]
+			]
 
-			SNew(SHorizontalBox)
 			+ SHorizontalBox::Slot()
-		.AutoWidth()
-		.VAlign(VAlign_Center)
-		.Padding(0, 0, 3, 0)
-		[
-			SNew(SOverlay)
-			+ SOverlay::Slot()
-		.HAlign(HAlign_Center)
-		.VAlign(VAlign_Top)
-		[
-			SNew(SImage)
-			.ColorAndOpacity_Lambda([this] { return GetStatusIconColor(); })
-		.Image_Lambda([this] { return GetStatusIcon(); })
-		.ToolTipText_Lambda([this] { return GetStatusToolTipText(); })
-		]
-		]
-
-	+ SHorizontalBox::Slot()
-		.AutoWidth()
-		.VAlign(VAlign_Center)
-		.Padding(0, 0, 10, 0)
-		[
-			SNew(STextBlock)
-			.Text_Lambda([this] { return GetTitleText(); })
-		.ToolTipText_Lambda([this] { return GetStatusToolTipText(); })
-		]
-		]
+				.AutoWidth()
+				.VAlign(VAlign_Center)
+				.Padding(0, 0, 10, 0)
+				[
+					SNew(STextBlock)
+					.Text_Lambda([this] { return GetTitleText(); })
+				.ToolTipText_Lambda([this] { return GetStatusToolTipText(); })
+				]
+	]
 	.OnGetMenuContent(FOnGetContent::CreateRaw(this, &SEditorPerformanceStatusBarWidget::CreateStatusBarMenu))
 		];
 
@@ -183,7 +181,7 @@ EActiveTimerReturnType SEditorPerformanceStatusBarWidget::UpdateState(double InC
 	WarningCount = 0;
 
 	// Check for KPIs that have exceeded their value
-	for (FKPIValues::TConstIterator It(EditorPerfModule.GetKPIValues()); It; ++It)
+	for (FKPIValues::TConstIterator It(EditorPerfModule.GetKPIRegistry().GetKPIValues()); It; ++It)
 	{
 		const FKPIValue& KPIValue = It->Value;
 
@@ -356,5 +354,3 @@ FText SEditorPerformanceStatusBarWidget::GetTitleToolTipText() const
 }
 
 #undef LOCTEXT_NAMESPACE
-
-UE_ENABLE_OPTIMIZATION_SHIP
diff --git a/Engine/Plugins/Experimental/EditorPerformance/Source/EditorPerformance/Public/EditorPerformanceModule.h b/Engine/Plugins/Experimental/EditorPerformance/Source/EditorPerformance/Public/EditorPerformanceModule.h
--- a/Engine/Plugins/Experimental/EditorPerformance/Source/EditorPerformance/Public/EditorPerformanceModule.h
+++ b/Engine/Plugins/Experimental/EditorPerformance/Source/EditorPerformance/Public/EditorPerformanceModule.h
@@ -36,7 +36,7 @@ public:
 
 	void UpdateKPIs();
 
-	const FKPIValues& GetKPIValues() const;
+	const FKPIRegistry& GetKPIRegistry() const;
 	const FString& GetKPIProfileName() const;
 
 	bool RecordInsightsSnaphshot(const FKPIValue& Value);
diff --git a/Engine/Plugins/Experimental/EditorPerformance/Source/EditorPerformance/Public/KPIValue.h b/Engine/Plugins/Experimental/EditorPerformance/Source/EditorPerformance/Public/KPIValue.h
--- a/Engine/Plugins/Experimental/EditorPerformance/Source/EditorPerformance/Public/KPIValue.h
+++ b/Engine/Plugins/Experimental/EditorPerformance/Source/EditorPerformance/Public/KPIValue.h
@@ -59,6 +59,7 @@ public:
 	static FString		GetValueAsString( float Value, FKPIValue::EDisplayType Type );
 	static FString		GetComparisonAsString(FKPIValue::ECompare Compare);
 	static FString		GetComparisonAsPrettyString(FKPIValue::ECompare Compare);
+	static FString		GetDisplayTypeAsString(FKPIValue::EDisplayType Type);
 
 	FGuid			Id;
 	FName			Category;
@@ -83,6 +84,15 @@ public:
 
 typedef TMap<FString, FKPIProfile> FKPIProfiles;
 
+class FKPIHint
+{
+public:
+	FString				Message;
+	FString				URL;
+};
+
+typedef TMap<FName, FKPIHint> FKPIHints;
+
 class FKPIRegistry
 {
 public:
@@ -93,9 +103,12 @@ public:
 	bool							SetKPIThreshold(const FName Name, float ThresholdValue);
 	bool							InvalidateKPIValue(const FName Name);
 	bool 							GetKPIValue(const FName Name, FKPIValue& Result) const;
+	bool 							GetKPIHint(const FName Name, FKPIHint& Result) const;
 	const FKPIValues&				GetKPIValues() const;
 	const FKPIProfiles&				GetKPIProfiles() const;
+	
 
+	void							LoadKPIHints(const FString& HintSectionName, const FString& FileName);
 	void							LoadKPIProfiles(const FString& ProfileSectionName, const FString& FileName);
 	bool							ApplyKPIProfile(const FKPIProfile& Profile);
 
@@ -103,5 +116,6 @@ private:
 
 	FKPIValues						Values;
 	FKPIProfiles					Profiles;
+	FKPIHints						Hints;
 };
 
