Editor Performance Tool
Added Vertical Scroll Bar
Added Documentation Link
Added Next Hint button
Added Open Settings

diff --git a/Engine/Plugins/Experimental/EditorPerformance/Source/EditorPerformance/Private/EditorPerformanceModule.cpp b/Engine/Plugins/Experimental/EditorPerformance/Source/EditorPerformance/Private/EditorPerformanceModule.cpp
--- a/Engine/Plugins/Experimental/EditorPerformance/Source/EditorPerformance/Private/EditorPerformanceModule.cpp
+++ b/Engine/Plugins/Experimental/EditorPerformance/Source/EditorPerformance/Private/EditorPerformanceModule.cpp
@@ -127,7 +127,7 @@ void FEditorPerformanceModule::ShowPerformanceReportTab()
 
 const FName EditorCategoryName = TEXT("Editor");
 const FName PIECategoryName = TEXT("PIE");
-const FName DDCCategoryName = TEXT("Cache");
+const FName CacheCategoryName = TEXT("Cache");
 const FName HardwareCategoryName = TEXT("Hardware");
 
 const FName EditorBootKPIName = TEXT("Boot");
@@ -172,21 +172,30 @@ void FEditorPerformanceModule::InitializeKPIs()
 	KPIRegistry.DeclareKPIValue(PIECategoryName, PIETransitionKPIName, 0.0, PIETransitionKPILimit, FKPIValue::LessThan, FKPIValue::Minutes);
 	KPIRegistry.DeclareKPIValue(PIECategoryName, PIEShutdownKPIName, 0.0, PIEShutdownKPILimit, FKPIValue::LessThan, FKPIValue::Minutes);
 	KPIRegistry.DeclareKPIValue(PIECategoryName, TotalTimeToPIEKPIName, 0.0, TotalTimeToPIEKPILimit, FKPIValue::LessThan, FKPIValue::Minutes);
-	KPIRegistry.DeclareKPIValue(DDCCategoryName, CloudDDCLatencyKPIName, 0.0, CloudDDCLatencyKPILimit, FKPIValue::LessThan, FKPIValue::Milliseconds);
-	KPIRegistry.DeclareKPIValue(DDCCategoryName, CloudDDCReadSpeedKPIName, 100.0, CloudDDCReadSpeedKPILimit, FKPIValue::GreaterThan, FKPIValue::MegaBitsPerSecond);
-	KPIRegistry.DeclareKPIValue(DDCCategoryName, TotalDDCEfficiencyKPIName, 100.0, TotalDDCEffciencyKPILimit, FKPIValue::GreaterThan, FKPIValue::Percent);
-	KPIRegistry.DeclareKPIValue(DDCCategoryName, LocalDDCEfficiencyKPIName, 100.0, LocalDDCEffciencyKPILimit, FKPIValue::GreaterThan, FKPIValue::Percent);
+	KPIRegistry.DeclareKPIValue(CacheCategoryName, CloudDDCLatencyKPIName, 0.0, CloudDDCLatencyKPILimit, FKPIValue::LessThan, FKPIValue::Milliseconds);
+	KPIRegistry.DeclareKPIValue(CacheCategoryName, CloudDDCReadSpeedKPIName, 100.0, CloudDDCReadSpeedKPILimit, FKPIValue::GreaterThan, FKPIValue::MegaBitsPerSecond);
+	KPIRegistry.DeclareKPIValue(CacheCategoryName, TotalDDCEfficiencyKPIName, 100.0, TotalDDCEffciencyKPILimit, FKPIValue::GreaterThan, FKPIValue::Percent);
+	KPIRegistry.DeclareKPIValue(CacheCategoryName, LocalDDCEfficiencyKPIName, 100.0, LocalDDCEffciencyKPILimit, FKPIValue::GreaterThan, FKPIValue::Percent);
 	//KPIRegistry.DeclareKPIValue(HitchrateKPIName, 0.0, HitchrateKPILimit, FKPIValue::LessThan, FKPIValue::Percent);
 	KPIRegistry.DeclareKPIValue(HardwareCategoryName, CoreCountKPIName, 128.0, CoreCountKPILimit, FKPIValue::GreaterThanOrEqual, FKPIValue::Decimal);
 	KPIRegistry.DeclareKPIValue(HardwareCategoryName, TotalMemoryKPIName, 128.0, TotalMemoryKPILimit, FKPIValue::GreaterThanOrEqual, FKPIValue::GigaBytes);
 	KPIRegistry.DeclareKPIValue(HardwareCategoryName, AvailableMemoryKPIName, 128.0, AvailableMemoryKPILimit, FKPIValue::GreaterThanOrEqual, FKPIValue::GigaBytes);
 
+	// Declare the KPI Hints
+	KPIRegistry.DeclareKPIHint(EditorCategoryName, EditorBootKPIName, LOCTEXT("EditorBootHintMessage", "The Editor boot time is slow.\nCheck you have enbabled a Game Feature Plugin profile for your project and that the expected local cache efficiency is met.\nIf you are booting the Editor in the background then disable the Use Less CPU in Background option in the settings."), LOCTEXT("EditorBootHintURL","https://docs.unrealengine.com/5.0/en-US/"));
+	KPIRegistry.DeclareKPIHint(EditorCategoryName, EditorStartupKPIName, LOCTEXT("EditorStartupHintMessage", "The Editor start - up time is slow.\nCheck you have enbabled a Game Feature Plugin profile for your project and that the expected local cache efficiency is met.\nIf you are booting the Editor in the background then disable the Use Less CPU in Background option in the settings."), LOCTEXT("EditorBootHintURL","https://docs.unrealengine.com/5.0/en-US/"));
+	
+	KPIRegistry.DeclareKPIHint(PIECategoryName, PIETransitionKPIName, LOCTEXT("PIETransitionHintMessage", "The Editor transtion to PIE is slow.\nCheck that the expected local cache efficiency is met.\nIf you are transitioning to PIE with the Editor in the background then disable the Use Less CPU in Background option in the settings."), LOCTEXT("PIETransitionHintURL", "https://docs.unrealengine.com/5.0/en-US/"));
+	
+	KPIRegistry.DeclareKPIHint(CacheCategoryName, LocalDDCEfficiencyKPIName, LOCTEXT("LocalCacheEfficencyHintMessage", "The Editor will not perform well if the local cache efficiency has not yet met the expected value.\nIf this is the first time you have booted the Editor after a sync then this is to be expected."), LOCTEXT("EditorCacheHintURL","https://docs.unrealengine.com/5.3/en-US/derived-data-cache/"));
+	
+	KPIRegistry.DeclareKPIHint(HardwareCategoryName, CoreCountKPIName, LOCTEXT("LowCoreCountHintMessage", "Your hardware has a low CPU core count.\nUsing a lower than recommended hardware specifcation for development is not recommended for good developer efficiency."), LOCTEXT("LowCoreCountHintURL", "https://docs.unrealengine.com/5.0/en-US/"));
+	KPIRegistry.DeclareKPIHint(HardwareCategoryName, TotalMemoryKPIName, LOCTEXT("LowTotalMemoryHintMessage", "Your hardware has a low Total Memory.\nUsing a lower than recommended hardware specifcation for development is not recommended for good developer efficiency."), LOCTEXT("LowTotalMemoryHintURL", "https://docs.unrealengine.com/5.0/en-US/"));
+	KPIRegistry.DeclareKPIHint(HardwareCategoryName, AvailableMemoryKPIName, LOCTEXT("LowAvaliableMemoryHintMessage", "You hardware is running low on memory.\nTry closing applciations that are no longer needed to recover available memory."), LOCTEXT("LowAvailableMemoryHintURL", "https://docs.unrealengine.com/5.0/en-US/"));
+
 	// Load the KPI profiles
 	KPIRegistry.LoadKPIProfiles(TEXT("EditorPerformance.Profile"), GEditorIni);
 
-	// Load the KPI hints
-	KPIRegistry.LoadKPIHints(TEXT("EditorPerformance.Hints"), GEditorIni);
-
 	// Apply any non map specific profiles
 	for (FKPIProfiles::TConstIterator It(KPIRegistry.GetKPIProfiles()); It; ++It)
 	{
@@ -291,31 +300,23 @@ void FEditorPerformanceModule::UpdateKPIs()
 	FDerivedDataCacheSummaryStats SummaryStats;
 	GatherDerivedDataCacheSummaryStats(SummaryStats);
 
-	int64 CloudGetHits = 0;
+	int64 TotalCloudGetHits = 0;
+	float CloudLatency = 0.0;
+	float CloudReadSpeed = 0.0;
 
 	for (const FDerivedDataCacheSummaryStat& Stat : SummaryStats.Stats)
 	{
 		if (Stat.Key == TEXT("CloudGetHits"))
 		{
-			CloudGetHits = FCString::Atoi(*Stat.Value);
+			TotalCloudGetHits = FCString::Atoi(*Stat.Value);
 		}
 		else if (Stat.Key == TEXT("CloudLatency"))
 		{
-			const float Value = FCString::Atof(*Stat.Value);
-
-			if (Value > 0.0f)
-			{
-				KPIRegistry.SetKPIValue(CloudDDCLatencyKPIName, Value);
-			}
+			CloudLatency = FCString::Atof(*Stat.Value);
 		}
 		else if (Stat.Key == TEXT("CloudReadSpeed"))
 		{
-			const float Value = FCString::Atof(*Stat.Value) * 8.0f;
-
-			if (Value > 0.0f)
-			{
-				KPIRegistry.SetKPIValue(CloudDDCReadSpeedKPIName, Value);
-			}
+			CloudReadSpeed = FCString::Atof(*Stat.Value) * 8.0f;
 		}
 		else if (Stat.Key == TEXT("TotalGetHitPct"))
 		{
@@ -337,12 +338,26 @@ void FEditorPerformanceModule::UpdateKPIs()
 		}
 	}
 
-	if (CloudGetHits == 0)
+	// Evaluate Cloud Cache performance
+	const int64 MinimalCloudGetHits = 10;
+	static int64 ElapsedCloudCacheHits = 0;
+	static int64 PreviousTotalCloudGetHits = 0;
+	static float AverageCloudLatency = 0;
+	static float AverageCloudReadSpeed = 0;
+
+	ElapsedCloudCacheHits = TotalCloudGetHits- PreviousTotalCloudGetHits;
+	PreviousTotalCloudGetHits = TotalCloudGetHits;
+	
+	if (TotalCloudGetHits < MinimalCloudGetHits)
 	{
 		KPIRegistry.InvalidateKPIValue(CloudDDCLatencyKPIName);
 		KPIRegistry.InvalidateKPIValue(CloudDDCReadSpeedKPIName);
 	}
-	
+	else
+	{
+		KPIRegistry.SetKPIValue(CloudDDCLatencyKPIName, CloudLatency);
+		KPIRegistry.SetKPIValue(CloudDDCReadSpeedKPIName, CloudReadSpeed);
+	}
 }
 
 bool FEditorPerformanceModule::IsHotLocalCacheCase() const
diff --git a/Engine/Plugins/Experimental/EditorPerformance/Source/EditorPerformance/Private/KPIValue.cpp b/Engine/Plugins/Experimental/EditorPerformance/Source/EditorPerformance/Private/KPIValue.cpp
--- a/Engine/Plugins/Experimental/EditorPerformance/Source/EditorPerformance/Private/KPIValue.cpp
+++ b/Engine/Plugins/Experimental/EditorPerformance/Source/EditorPerformance/Private/KPIValue.cpp
@@ -249,6 +249,35 @@ bool FKPIRegistry::DeclareKPIValue( const FKPIValue& Value )
 	return false;
 }
 
+bool FKPIRegistry::DeclareKPIHint(const FName Category, const FName Name, const FText& HintMessage, const FText& HintURL)
+{
+	if (Values.Find(Name) != nullptr)
+	{
+		FKPIHint NewKPIHint;
+
+		NewKPIHint.Name = Name;
+		NewKPIHint.Category = Category;
+		NewKPIHint.Message = HintMessage;
+		NewKPIHint.URL = HintURL;
+
+		FKPIHint* KPIHint = Hints.Find(Name);
+
+		if (KPIHint != nullptr)
+		{
+			*KPIHint = NewKPIHint;
+			return true;
+		}
+		else
+		{
+			Hints.Emplace(Name, NewKPIHint);
+			return true;
+		}
+	}
+
+	return false;
+}
+
+
 bool FKPIRegistry::InvalidateKPIValue(const FName Name)
 {
 	FKPIValue* ExistingValue = Values.Find(Name);
@@ -329,7 +358,7 @@ const FKPIProfiles& FKPIRegistry::GetKPIProfiles() const
 
 void FKPIRegistry::LoadKPIHints(const FString& HintSectionName, const FString& FileName)
 {
-	TArray<FString> SectionNames;
+	/*TArray<FString> SectionNames;
 
 	if (GConfig->GetSectionNames(FileName, SectionNames))
 	{
@@ -338,7 +367,7 @@ void FKPIRegistry::LoadKPIHints(const FString& HintSectionName, const FString& F
 			if (SectionName.Find(HintSectionName) != INDEX_NONE)
 			{
 				FKPIHint Hint;
-				Hint.URL = TEXT("https://docs.unrealengine.com/5.0/en-US/");
+				Hint.URL = FText(TEXT("https://docs.unrealengine.com/5.0/en-US/"));
 
 				for (FKPIValues::TConstIterator It(GetKPIValues()); It; ++It)
 				{
@@ -352,7 +381,7 @@ void FKPIRegistry::LoadKPIHints(const FString& HintSectionName, const FString& F
 				}
 			}
 		}
-	}
+	}*/
 }
 
 void FKPIRegistry::LoadKPIProfiles(const FString& ProfileSectionName, const FString& FileName)
diff --git a/Engine/Plugins/Experimental/EditorPerformance/Source/EditorPerformance/Private/SEditorPerformanceDialogs.cpp b/Engine/Plugins/Experimental/EditorPerformance/Source/EditorPerformance/Private/SEditorPerformanceDialogs.cpp
--- a/Engine/Plugins/Experimental/EditorPerformance/Source/EditorPerformance/Private/SEditorPerformanceDialogs.cpp
+++ b/Engine/Plugins/Experimental/EditorPerformance/Source/EditorPerformance/Private/SEditorPerformanceDialogs.cpp
@@ -16,7 +16,10 @@
 #include "Widgets/Text/STextBlock.h"
 #include "Widgets/Input/SComboBox.h"
 #include "Widgets/Input/SCheckBox.h"
+#include "Widgets/Input/SHyperlink.h"
 #include "EditorPerformanceModule.h"
+#include "Modules/ModuleManager.h"
+#include "ISettingsModule.h"
 #include "Editor/EditorPerformanceSettings.h"
 
 #define LOCTEXT_NAMESPACE "EditorPerformance"
@@ -34,47 +37,52 @@ void SEditorPerformanceReportDialog::Construct(const FArguments& InArgs)
 
 	this->ChildSlot
 	[
-		SNew(SVerticalBox)
-		+ SVerticalBox::Slot()
-		.Padding(0, 20, 0, 0)
-		.AutoHeight()
+		SNew(SScrollBox)
+		.Orientation(Orient_Vertical)
+		+ SScrollBox::Slot()
 		[
-			SNew(SHorizontalBox)
-			+SHorizontalBox::Slot()
-			.FillWidth(1.0f)
+			SNew(SVerticalBox)
+			+ SVerticalBox::Slot()
+			.Padding(0, 20, 0, 0)
+			.AutoHeight()
 			[
-				SNew(STextBlock)
-				.Margin(TitleMargin)
-				.ColorAndOpacity(TitleColor)
-				.Font(TitleFont)
-				.Justification(ETextJustify::Left)
-				.Text_Lambda([this,&EditorPerfModule]
-					{ 
-						return FText::FromString(*FString::Printf(TEXT("Profile : %s"), *EditorPerfModule.GetKPIProfileName())); 
-					}
-				)
+				SNew(SHorizontalBox)
+				+SHorizontalBox::Slot()
+				.FillWidth(1.0f)
+				[
+					SNew(STextBlock)
+					.Margin(TitleMargin)
+					.ColorAndOpacity(TitleColor)
+					.Font(TitleFont)
+					.Justification(ETextJustify::Left)
+					.Text_Lambda([this,&EditorPerfModule]
+						{ 
+							return FText::FromString(*FString::Printf(TEXT("Profile : %s"), *EditorPerfModule.GetKPIProfileName())); 
+						}
+					)
+				]
+			]
+			+ SVerticalBox::Slot()
+			.AutoHeight()
+			.Padding(0, 5, 0, 0)
+			.Expose(KPIGridSlot)
+			[
+				GetKPIGridPanel()
+			]
+			+ SVerticalBox::Slot()
+			.AutoHeight()
+			.Padding(0, 5, 0, 0)
+			.Expose(HintGridSlot)
+			[
+				GetHintGridPanel()
+			]
+			+ SVerticalBox::Slot()
+			.AutoHeight()
+			.Padding(0, 5, 0, 0)
+			.Expose(SettingsGridSlot)
+			[
+				GetSettingsGridPanel()
 			]
-		]
-		+ SVerticalBox::Slot()
-		.AutoHeight()
-		.Padding(0, 5, 0, 0)
-		.Expose(SettingsGridSlot)
-		[
-			GetSettingsGridPanel()
-		]
-		+ SVerticalBox::Slot()
-		.AutoHeight()
-		.Padding(0, 5, 0, 0)
-		.Expose(KPIGridSlot)
-		[
-			GetKPIGridPanel()
-		]
-		+ SVerticalBox::Slot()
-		.AutoHeight()
-		.Padding(0, 5, 0, 0)
-		.Expose(HintGridSlot)
-		[
-			GetHintGridPanel()
 		]
 	];
 
@@ -120,17 +128,17 @@ TSharedRef<SWidget> SEditorPerformanceReportDialog::GetSettingsGridPanel()
 
 	int32 Row = 0;
 
-	/*Panel->AddSlot(0, Row)
+	Panel->AddSlot(0, Row)
 		.HAlign(HAlign_Left)
 		[
 			SNew(STextBlock)
 			.Margin(TitleMarginFirstColumn)
 			.ColorAndOpacity(TitleColor)
 			.Font(TitleFont)
-			.Text(LOCTEXT("Settings", "Settings"))
+			.Text(LOCTEXT("SettingsText", "Settings"))
 		];
 
-	Row++;*/
+	Row++;
 
 	Panel->AddSlot(0, Row)
 		.HAlign(HAlign_Left)
@@ -157,7 +165,8 @@ TSharedRef<SWidget> SEditorPerformanceReportDialog::GetSettingsGridPanel()
 			.Padding(FMargin(4.0f, 0.0f))
 			[
 				SNew(STextBlock)
-				.Text(LOCTEXT("EnableNotifications", "Enable Notifications"))
+				.Text(LOCTEXT("EnableNotificationsText", "Notifications"))
+				.ToolTipText(LOCTEXT("EnableNotificationsToolTip", "Enable All Notifications"))
 				.ColorAndOpacity(EStyleColor::Foreground)
 			]
 		];
@@ -187,7 +196,8 @@ TSharedRef<SWidget> SEditorPerformanceReportDialog::GetSettingsGridPanel()
 				.Padding(FMargin(4.0f, 0.0f))
 				[
 					SNew(STextBlock)
-					.Text(LOCTEXT("EnableSnapshots", "Enable Snapshots"))
+					.Text(LOCTEXT("EnableSnapshotsText", "Snapshots"))
+					.ToolTipText(LOCTEXT("EnableSnapshotsToolTip", "Enable Automatic Capture of Unreal Insights Snaphsots"))
 					.ColorAndOpacity(EStyleColor::Foreground)
 				]
 			];
@@ -214,26 +224,85 @@ TSharedRef<SWidget> SEditorPerformanceReportDialog::GetSettingsGridPanel()
 
 				UpdateGridPanels(0.0f, 0.0f);
 			})
-					.Padding(FMargin(4.0f, 0.0f))
-				[
-					SNew(STextBlock)
-					.Text(LOCTEXT("EnableTelemetry", "Enable Telemetry"))
-					.ColorAndOpacity(EStyleColor::Foreground)
-				]
+			.Padding(FMargin(4.0f, 0.0f))
+			[
+				SNew(STextBlock)
+				.Text(LOCTEXT("EnableTelemetryText", "Telemetry"))
+				.ToolTipText(LOCTEXT("EnableTelemetryToolTip", "Record Warning Telemetry Events To Analytics System"))
+				.ColorAndOpacity(EStyleColor::Foreground)
+			]
 		];
 
-	Row++;
+	Panel->AddSlot(3, Row)
+		.HAlign(HAlign_Left)
+		[
+			SNew(SCheckBox)
+			.IsChecked_Lambda([]
+				{
+					const UEditorPerformanceSettings* EditorPerformanceSettings = GetDefault<UEditorPerformanceSettings>();
+					return EditorPerformanceSettings && EditorPerformanceSettings->bThrottleCPUWhenNotForeground ? ECheckBoxState::Checked : ECheckBoxState::Unchecked;
+				})
+			.OnCheckStateChanged_Lambda([this](ECheckBoxState NewState)
+			{
+				UEditorPerformanceSettings* EditorPerformanceSettings = GetMutableDefault<UEditorPerformanceSettings>();
 
-	return Panel;
-}
+				if (EditorPerformanceSettings)
+				{
+					EditorPerformanceSettings->bThrottleCPUWhenNotForeground = NewState == ECheckBoxState::Checked;
+					EditorPerformanceSettings->PostEditChange();
+					EditorPerformanceSettings->SaveConfig();
+				}
 
-void SEditorPerformanceReportDialog::OpenURL(const FString& URL) const
-{
-	FPlatformProcess::LaunchURL(*URL, nullptr, nullptr);
+				UpdateGridPanels(0.0f, 0.0f);
+			})
+			.Padding(FMargin(4.0f, 0.0f))
+			[
+				SNew(STextBlock)
+				.Text(LOCTEXT("EnableBackgroundThrottlingText", "Throttling"))
+				.ToolTipText(LOCTEXT("EnableBackgroundThrottlingToolTip", "Enable CPU throttling when the Editor is in the background."))
+				.ColorAndOpacity(EStyleColor::Foreground)
+			]
+		];
+
+	Panel->AddSlot(4, Row)
+		.HAlign(HAlign_Left)
+		.Padding(FMargin(10.0f, 10.0f))
+		[
+			SNew(SButton)
+			.Text(LOCTEXT("OpenSettingsText", "All Settings"))
+			.ToolTipText(LOCTEXT("OpenSettingsToolTip", "Open the Editor Performance Settings Tab."))
+			.OnClicked_Lambda([this]()
+			{
+				FModuleManager::LoadModuleChecked<ISettingsModule>("Settings").ShowViewer("Editor", "General", "EditorPerformanceSettings");
+				return FReply::Handled();
+			})
+		];
+
+	return Panel;
 }
 
 TSharedRef<SWidget> SEditorPerformanceReportDialog::GetHintGridPanel()
 {
+	FEditorPerformanceModule& EditorPerfModule = FModuleManager::LoadModuleChecked<FEditorPerformanceModule>("EditorPerformance");
+
+	int32 NumHints = 0;
+
+	TArray<FKPIHint> KPIHints;
+
+	const FKPIValues& KPIValues = EditorPerfModule.GetKPIRegistry().GetKPIValues();
+
+	for (FKPIValues::TConstIterator It(KPIValues); It; ++It)
+	{
+		const FKPIValue& KPIValue = It->Value;
+
+		FKPIHint KPIHint;
+
+		if (KPIValue.State == FKPIValue::Bad && EditorPerfModule.GetKPIRegistry().GetKPIHint(KPIValue.Name, KPIHint))
+		{
+			KPIHints.Emplace(KPIHint);
+		}
+	}
+
 	TSharedRef<SGridPanel> Panel =
 		SNew(SGridPanel);
 
@@ -247,68 +316,84 @@ TSharedRef<SWidget> SEditorPerformanceReportDialog::GetHintGridPanel()
 	const FMargin DefaultMargin(0.0f, RowMargin, ColumnMargin, RowMargin);
 	const FMargin DefaultMarginFirstColumn(ColumnMargin, RowMargin);
 
-	FEditorPerformanceModule& EditorPerfModule = FModuleManager::LoadModuleChecked<FEditorPerformanceModule>("EditorPerformance");
-
 	int32 Row = 0;
 
-	FKPIHint KPIHint;
-
-	for (FKPIValues::TConstIterator It(EditorPerfModule.GetKPIRegistry().GetKPIValues()); It; ++It)
+	if (KPIHints.Num()>0)
 	{
-		const FKPIValue& KPIValue = It->Value;
+		CurrentHintIndex = CurrentHintIndex % KPIHints.Num();
 
-		if (KPIValue.State==FKPIValue::Bad && EditorPerfModule.GetKPIRegistry().GetKPIHint(KPIValue.Name, KPIHint))
-		{
-			Panel->AddSlot(0, Row)
-				.HAlign(HAlign_Left)
-				[
-					SNew(STextBlock)
-					.Margin(TitleMarginFirstColumn)
-					.ColorAndOpacity(TitleColor)
-					.Font(TitleFont)
-					.Justification(ETextJustify::Left)
-					.Text(LOCTEXT("HintTitle", "Hint"))
-				];
+		const FKPIHint& KPIHint = KPIHints[CurrentHintIndex];
 
-			Row++;
+		Panel->AddSlot(0, Row)
+			.HAlign(HAlign_Left)
+			[
+				SNew(STextBlock)
+				.Margin(TitleMarginFirstColumn)
+				.ColorAndOpacity(TitleColor)
+				.Font(TitleFont)
+				.Justification(ETextJustify::Left)
+				.Text(LOCTEXT("HintsTitle", "Hints"))
+			];
 
-			Panel->AddSlot(0, Row)
-				.HAlign(HAlign_Left)
-				[
-					SNew(STextBlock)
-					.Margin(TitleMarginFirstColumn)
-					.ColorAndOpacity(TitleColor)
-					.Font(TitleFont)
-					.Text(FText::FromString(*FString::Printf(TEXT("%s %s"), *KPIValue.Category.ToString(), *KPIValue.Name.ToString())))
-				];
+		Row++;
 
-			Row++;
+		Panel->AddSlot(0, Row)
+			.HAlign(HAlign_Left)
+			[
+				SNew(STextBlock)
+				.Margin(TitleMarginFirstColumn)
+				.ColorAndOpacity(EStyleColor::Foreground)
+				.Font(TitleFont)
+				.Text(FText::FromString(*FString::Printf(TEXT("%s %s"), *KPIHint.Category.ToString(), *KPIHint.Name.ToString())))
+			];
+
+		Row++;
+
+		Panel->AddSlot(0, Row)
+			.HAlign(HAlign_Left)
+			[
+				SNew(STextBlock)
+				.AutoWrapText(true)
+				.Margin(DefaultMarginFirstColumn)
+				.ColorAndOpacity(EStyleColor::Foreground)
+				.Justification(ETextJustify::Left)
+				.Text(KPIHint.Message)
+			];
 
+		Row++;
+
+		if (!KPIHint.URL.IsEmpty())
+		{
 			Panel->AddSlot(0, Row)
 				.HAlign(HAlign_Left)
+				.Padding(FMargin(10.0f, 10.0f))
 				[
-					SNew(STextBlock)
-					.AutoWrapText(true)
-					.Margin(DefaultMarginFirstColumn)
-					.ColorAndOpacity(EStyleColor::Foreground)
-					.Justification(ETextJustify::Left)
-					.Text(FText::FromString(*KPIHint.Message))
+					SNew(SHyperlink)
+					.Text(LOCTEXT("HintLinkName", "Further Help & Documentation"))
+					.ToolTipText_Lambda([=]() { return FText::FromString(*KPIHint.URL.ToString()); })
+					.OnNavigate_Lambda([=]() { FPlatformProcess::LaunchURL(*KPIHint.URL.ToString(), nullptr, nullptr); })
 				];
 
 			Row++;
+		}
 
+		if (KPIHints.Num() > 1)
+		{
 			Panel->AddSlot(0, Row)
 				.HAlign(HAlign_Left)
+				.Padding(FMargin(10.0f, 10.0f))
 				[
-					SNew(STextBlock)
-					.AutoWrapText(true)
-					.Margin(DefaultMarginFirstColumn)
-					.ColorAndOpacity(EStyleColor::AccentBlue)
-					.Justification(ETextJustify::Left)
-					.Text(FText::FromString(*KPIHint.URL))
+					SNew(SButton)
+					.Text(LOCTEXT("NextHintName", "Next Hint"))
+					.OnClicked_Lambda([this]()
+					{
+						CurrentHintIndex++;
+						UpdateGridPanels(0.0f, 0.0f);
+						return FReply::Handled();
+					})
 				];
 
-			break;
+			Row++;
 		}
 	}
 
@@ -368,6 +453,8 @@ TSharedRef<SWidget> SEditorPerformanceReportDialog::GetKPIGridPanel()
 				if (EditorPerformanceSettings)
 				{
 					EditorPerformanceSettings->bShowWarningsOnly = (Name == WarningFilterOptions[1]);
+					EditorPerformanceSettings->PostEditChange();
+					EditorPerformanceSettings->SaveConfig();
 				}
 
 				UpdateGridPanels(0.0f, 0.0f);
@@ -386,7 +473,7 @@ TSharedRef<SWidget> SEditorPerformanceReportDialog::GetKPIGridPanel()
 		[
 			SNew(STextBlock)
 			.Margin(DefaultMargin)
-			.ColorAndOpacity(TitleColor)
+			.ColorAndOpacity(EStyleColor::Foreground)
 			.Font(TitleFont)
 			.Text(LOCTEXT("CurrentValueColumn", "Current"))
 		];
@@ -396,7 +483,7 @@ TSharedRef<SWidget> SEditorPerformanceReportDialog::GetKPIGridPanel()
 		[
 			SNew(STextBlock)
 			.Margin(DefaultMargin)
-			.ColorAndOpacity(TitleColor)
+			.ColorAndOpacity(EStyleColor::Foreground)
 			.Font(TitleFont)
 			.Text(LOCTEXT("ExpectedValueColumn", "Expected"))
 		];
@@ -408,7 +495,7 @@ TSharedRef<SWidget> SEditorPerformanceReportDialog::GetKPIGridPanel()
 			[
 				SNew(STextBlock)
 				.Margin(DefaultMargin)
-				.ColorAndOpacity(TitleColor)
+				.ColorAndOpacity(EStyleColor::Foreground)
 				.Font(TitleFont)
 				.Text(LOCTEXT("NotifyColumn", "Notify"))
 			];
@@ -452,7 +539,7 @@ TSharedRef<SWidget> SEditorPerformanceReportDialog::GetKPIGridPanel()
 			[
 				SNew(STextBlock)
 				.Margin(TitleMarginFirstColumn)
-				.ColorAndOpacity(TitleColor)
+				.ColorAndOpacity(EStyleColor::Foreground)
 				.Font(TitleFont)
 				.Text(FText::FromString(*Category.ToString()))
 			];
diff --git a/Engine/Plugins/Experimental/EditorPerformance/Source/EditorPerformance/Private/SEditorPerformanceDialogs.h b/Engine/Plugins/Experimental/EditorPerformance/Source/EditorPerformance/Private/SEditorPerformanceDialogs.h
--- a/Engine/Plugins/Experimental/EditorPerformance/Source/EditorPerformance/Private/SEditorPerformanceDialogs.h
+++ b/Engine/Plugins/Experimental/EditorPerformance/Source/EditorPerformance/Private/SEditorPerformanceDialogs.h
@@ -4,6 +4,7 @@
 
 #include "Templates/SharedPointer.h"
 #include "Types/SlateEnums.h"
+#include "Input/Reply.h"
 #include "Widgets/DeclarativeSyntaxSupport.h"
 #include "Widgets/SBoxPanel.h"
 #include "Widgets/SCompoundWidget.h"
@@ -20,8 +21,7 @@ class SEditorPerformanceReportDialog : public SCompoundWidget
 private:
 	
 	EActiveTimerReturnType UpdateGridPanels(double InCurrentTime, float InDeltaTime);
-	void OpenURL(const FString& URL) const;
-
+	
 	TSharedRef<SWidget> GetKPIGridPanel();
 	TSharedRef<SWidget> GetSettingsGridPanel();
 	TSharedRef<SWidget> GetHintGridPanel();
@@ -29,8 +29,9 @@ private:
 	SVerticalBox::FSlot* SettingsGridSlot = nullptr;
 	SVerticalBox::FSlot* KPIGridSlot = nullptr;
 	SVerticalBox::FSlot* HintGridSlot = nullptr;
+	uint32 CurrentHintIndex=0;
 
 	const TArray<FName> NotifcationOptions = { FName("Notify"), FName("Ignore") };
-	const TArray<FName> WarningFilterOptions = { FName("Show All"), FName("Show Warnings") };
+	const TArray<FName> WarningFilterOptions = { FName("Show All"), FName("Warnings Only") };
 };
 
diff --git a/Engine/Plugins/Experimental/EditorPerformance/Source/EditorPerformance/Public/EditorPerformanceModule.h b/Engine/Plugins/Experimental/EditorPerformance/Source/EditorPerformance/Public/EditorPerformanceModule.h
--- a/Engine/Plugins/Experimental/EditorPerformance/Source/EditorPerformance/Public/EditorPerformanceModule.h
+++ b/Engine/Plugins/Experimental/EditorPerformance/Source/EditorPerformance/Public/EditorPerformanceModule.h
@@ -54,7 +54,6 @@ private:
 
 	void InitializeKPIs();
 	void TerminateKPIs();
-	void LoadKPIProfiles();
 
 	FKPIRegistry					KPIRegistry;
 	TMap<FString,FKPIProfile>		KPIProfiles;
diff --git a/Engine/Plugins/Experimental/EditorPerformance/Source/EditorPerformance/Public/KPIValue.h b/Engine/Plugins/Experimental/EditorPerformance/Source/EditorPerformance/Public/KPIValue.h
--- a/Engine/Plugins/Experimental/EditorPerformance/Source/EditorPerformance/Public/KPIValue.h
+++ b/Engine/Plugins/Experimental/EditorPerformance/Source/EditorPerformance/Public/KPIValue.h
@@ -87,8 +87,10 @@ typedef TMap<FString, FKPIProfile> FKPIProfiles;
 class FKPIHint
 {
 public:
-	FString				Message;
-	FString				URL;
+	FName			Category;
+	FName			Name;
+	FText			Message;
+	FText			URL;
 };
 
 typedef TMap<FName, FKPIHint> FKPIHints;
@@ -97,8 +99,10 @@ class FKPIRegistry
 {
 public:
 
-	bool							DeclareKPIValue(FName NewCategory, const FName Name, float InitialValue, float ThresholdValue, FKPIValue::ECompare Compare, FKPIValue::EDisplayType Type);
+	bool							DeclareKPIValue(const FName Category, const FName Name, float InitialValue, float ThresholdValue, FKPIValue::ECompare Compare, FKPIValue::EDisplayType Type);
 	bool							DeclareKPIValue(const FKPIValue& Value);
+	bool							DeclareKPIHint(const FName Category, const FName Name, const FText& HintMessage, const FText& HintURL);
+
 	bool							SetKPIValue(const FName Name, float CurrentValue);
 	bool							SetKPIThreshold(const FName Name, float ThresholdValue);
 	bool							InvalidateKPIValue(const FName Name);
