Changed Editor Performance Telemetry Event Name from Editor.PerformanceWarning to Editor.Performance.Warning so it is in keeping with the existing namign conventions
This event is now ready for schematization

diff --git a/Engine/Plugins/Experimental/EditorPerformance/Source/EditorPerformance/Private/EditorPerformanceModule.cpp b/Engine/Plugins/Experimental/EditorPerformance/Source/EditorPerformance/Private/EditorPerformanceModule.cpp
--- a/Engine/Plugins/Experimental/EditorPerformance/Source/EditorPerformance/Private/EditorPerformanceModule.cpp
+++ b/Engine/Plugins/Experimental/EditorPerformance/Source/EditorPerformance/Private/EditorPerformanceModule.cpp
@@ -413,7 +413,7 @@ bool FEditorPerformanceModule::RecordTelemetryEvent(const FKPIValue& KPIValue)
 		Attributes.Emplace(TEXT("KPI_DisplayType"), *FKPIValue::GetDisplayTypeAsString(KPIValue.DisplayType));
 		Attributes.Emplace(TEXT("KPI_Profile"), KPIProfileName);
 		
-		FStudioTelemetry::Get().RecordEvent(TEXT("Editor.PerformanceWarning"), Attributes);
+		FStudioTelemetry::Get().RecordEvent(TEXT("Editor.Performance.Warning"), Attributes);
 
 		return true;
 	}
