Removed Unreal Cloud DDC KPIS from the Editor Performance Tool until the source values become more reliable

diff --git a/Engine/Plugins/Experimental/EditorPerformance/Source/EditorPerformance/Private/EditorPerformanceModule.cpp b/Engine/Plugins/Experimental/EditorPerformance/Source/EditorPerformance/Private/EditorPerformanceModule.cpp
--- a/Engine/Plugins/Experimental/EditorPerformance/Source/EditorPerformance/Private/EditorPerformanceModule.cpp
+++ b/Engine/Plugins/Experimental/EditorPerformance/Source/EditorPerformance/Private/EditorPerformanceModule.cpp
@@ -172,8 +172,8 @@ void FEditorPerformanceModule::InitializeKPIs()
 	KPIRegistry.DeclareKPIValue(PIECategoryName, PIETransitionKPIName, 0.0, PIETransitionKPILimit, FKPIValue::LessThan, FKPIValue::Minutes);
 	KPIRegistry.DeclareKPIValue(PIECategoryName, PIEShutdownKPIName, 0.0, PIEShutdownKPILimit, FKPIValue::LessThan, FKPIValue::Minutes);
 	KPIRegistry.DeclareKPIValue(PIECategoryName, TotalTimeToPIEKPIName, 0.0, TotalTimeToPIEKPILimit, FKPIValue::LessThan, FKPIValue::Minutes);
-	KPIRegistry.DeclareKPIValue(CacheCategoryName, CloudDDCLatencyKPIName, 0.0, CloudDDCLatencyKPILimit, FKPIValue::LessThan, FKPIValue::Milliseconds);
-	KPIRegistry.DeclareKPIValue(CacheCategoryName, CloudDDCReadSpeedKPIName, 100.0, CloudDDCReadSpeedKPILimit, FKPIValue::GreaterThan, FKPIValue::MegaBitsPerSecond);
+	//KPIRegistry.DeclareKPIValue(CacheCategoryName, CloudDDCLatencyKPIName, 0.0, CloudDDCLatencyKPILimit, FKPIValue::LessThan, FKPIValue::Milliseconds);
+	//KPIRegistry.DeclareKPIValue(CacheCategoryName, CloudDDCReadSpeedKPIName, 100.0, CloudDDCReadSpeedKPILimit, FKPIValue::GreaterThan, FKPIValue::MegaBitsPerSecond);
 	KPIRegistry.DeclareKPIValue(CacheCategoryName, TotalDDCEfficiencyKPIName, 100.0, TotalDDCEffciencyKPILimit, FKPIValue::GreaterThan, FKPIValue::Percent);
 	KPIRegistry.DeclareKPIValue(CacheCategoryName, LocalDDCEfficiencyKPIName, 100.0, LocalDDCEffciencyKPILimit, FKPIValue::GreaterThan, FKPIValue::Percent);
 	//KPIRegistry.DeclareKPIValue(HitchrateKPIName, 0.0, HitchrateKPILimit, FKPIValue::LessThan, FKPIValue::Percent);
