Fixes for Studio Telemetry
Fixed Local DDC stats so they correctly count Zen Local DDC as a Local DDC if a file based Local DDC is not present
Added explicit AssetRegistryScan Start/Stop callbacks and removed the less reliable sample based method.
Added AssetRegistryScan as a KPI for the Editor Performance UI

diff --git a/Engine/Plugins/Experimental/EditorPerformance/Source/EditorPerformance/Private/EditorPerformanceModule.cpp b/Engine/Plugins/Experimental/EditorPerformance/Source/EditorPerformance/Private/EditorPerformanceModule.cpp
--- a/Engine/Plugins/Experimental/EditorPerformance/Source/EditorPerformance/Private/EditorPerformanceModule.cpp
+++ b/Engine/Plugins/Experimental/EditorPerformance/Source/EditorPerformance/Private/EditorPerformanceModule.cpp
@@ -16,6 +16,7 @@
 #include "Editor.h"
 #include "ToolMenus.h"
 #include "Editor/EditorPerformanceSettings.h"
+#include "AssetRegistry/AssetRegistryModule.h"
 #include "DerivedDataCacheUsageStats.h"
 #include "Trace/Trace.h"
 #include "StudioTelemetry.h"
@@ -35,6 +36,7 @@ const FName EditorBootKPIName = TEXT("Boot");
 const FName EditorInitializeKPIName = TEXT("Initialize");
 const FName EditorLoadMapKPIName = TEXT("Load Map");
 const FName EditorHitchrateKPIName = TEXT("Hitch Rate");
+const FName EditorAssetRegistryScanKPIName = TEXT("Asset Registry Scan");
 const FName TotalTimeToEditorKPIName = TEXT("Total Time To Editor");
 const FName TotalTimeToPIEKPIName = TEXT("Total Time To PIE");
 const FName PIEFirstTransitionKPIName = TEXT("First Transition");
@@ -53,6 +55,7 @@ float EditorBootKPILimit = 100;
 float EditorInitializeKPILimit = 160;
 float EditorLoadMapKPILimit = 120;
 float EditorHitchrateKPILimit = 25;
+float EditorAssetRegistryScanKPILimit = 60;
 float TotalTimeToEditorKPILimit = 160;
 float PIEFirstTransitionKPILimit = 220;
 float PIETransitionKPILimit = 40;
@@ -190,6 +193,7 @@ void FEditorPerformanceModule::InitializeKPIs()
 	EditorBootKPI			= KPIRegistry.DeclareKPIValue(EditorCategoryName, EditorBootKPIName, 0.0, EditorBootKPILimit, FKPIValue::LessThan, FKPIValue::Minutes);
 	EditorInitializeKPI		= //KPIRegistry.DeclareKPIValue(EditorCategoryName, EditorInitializeKPIName, 0.0, EditorInitializeKPILimit, FKPIValue::LessThan, FKPIValue::Minutes);
 	EditorLoadMapKPI		= KPIRegistry.DeclareKPIValue(EditorCategoryName, EditorLoadMapKPIName, 0.0, EditorLoadMapKPILimit, FKPIValue::LessThan, FKPIValue::Minutes);
+	AssetRegistryScanKPI	= KPIRegistry.DeclareKPIValue(EditorCategoryName, EditorAssetRegistryScanKPIName, 0.0, EditorAssetRegistryScanKPILimit, FKPIValue::LessThan, FKPIValue::Minutes);
 	TotalTimeToEditorKPI	= KPIRegistry.DeclareKPIValue(EditorCategoryName, TotalTimeToEditorKPIName, 0.0, TotalTimeToEditorKPILimit, FKPIValue::LessThan, FKPIValue::Minutes);
 	PIEFirstTransitionKPI	= KPIRegistry.DeclareKPIValue(PIECategoryName, PIEFirstTransitionKPIName, 0.0, PIEFirstTransitionKPILimit, FKPIValue::LessThan, FKPIValue::Minutes);
 	PIETransitionKPI		= KPIRegistry.DeclareKPIValue(PIECategoryName, PIETransitionKPIName, 0.0, PIETransitionKPILimit, FKPIValue::LessThan, FKPIValue::Minutes);
@@ -232,6 +236,8 @@ void FEditorPerformanceModule::InitializeKPIs()
 	KPIRegistry.SetKPIValue(CoreCountKPI, float(FPlatformMisc::NumberOfCores()));
 	KPIRegistry.SetKPIValue(TotalMemoryKPI, static_cast<float>(FPlatformMemory::GetStats().TotalPhysical) / (1024.0f * 1024.0f * 1024.0f));
 
+	EditorState = EEditorState::Editor_Boot;
+
 	// Register the delegates
 	FEditorDelegates::OnEditorBoot.AddLambda([this](double TimeToBootEditor )
 		{
@@ -256,13 +262,16 @@ void FEditorPerformanceModule::InitializeKPIs()
 	FEditorDelegates::OnMapLoad.AddLambda([this](const FString& MapName, FCanLoadMap& OutCanLoadMap)
 		{
 			LoadMapStartTime = FDateTime::UtcNow();
+			IsLoadingMap = true;
 		});
 
 	FEditorDelegates::OnMapOpened.AddLambda([this](const FString& MapName, bool Unused)
 		{
+			IsLoadingMap = false;
+
 			EditorMapName = FPaths::GetBaseFilename(MapName);
+	
 			EditorLoadMapTime = float((FDateTime::UtcNow() - LoadMapStartTime).GetTotalSeconds());
-			BootToPIETime += EditorLoadMapTime;
 			KPIRegistry.SetKPIValue(EditorLoadMapKPI, EditorLoadMapTime);
 	
 			// Apply any profile that matches the currently loaded map
@@ -281,21 +290,19 @@ void FEditorPerformanceModule::InitializeKPIs()
 	FEditorDelegates::StartPIE.AddLambda([this](bool)
 		{
 			PIEStartTime = FDateTime::UtcNow();
-			EditorState = EEditorState::PIE_Transtion;
+			EditorState = EEditorState::PIE_Startup;
 		});
 
 	FWorldDelegates::OnPIEReady.AddLambda([this](UGameInstance* GameInstance)
-		{
-			static bool IsFirstTimeToPIE = true;
-
+		{	
 			const float PIETransitionTime = float((FDateTime::UtcNow() - PIEStartTime).GetTotalSeconds());
 
 			if (IsFirstTimeToPIE)
 			{
-				BootToPIETime = EditorStartUpTime +  EditorLoadMapTime + PIETransitionTime;
+				BootToPIETime = EditorStartUpTime + EditorLoadMapTime + PIETransitionTime;
+				KPIRegistry.SetKPIValue(TotalTimeToPIEKPI, BootToPIETime);
 
 				KPIRegistry.SetKPIValue(PIEFirstTransitionKPI, PIETransitionTime);
-				KPIRegistry.SetKPIValue(TotalTimeToPIEKPI, BootToPIETime);
 			}
 			else
 			{
@@ -322,6 +329,47 @@ void FEditorPerformanceModule::InitializeKPIs()
 			EditorState = EEditorState::Editor_Interact;
 			EditorHitchCount = 0;
 		});
+
+	FModuleManager::Get().OnModulesChanged().AddLambda([this](FName ModuleName, EModuleChangeReason ChangeReason)
+		{
+			switch ( ChangeReason )
+			{ 
+				default:
+				{
+					break;
+				}
+				
+				case EModuleChangeReason::ModuleLoaded:
+				{
+					TotalPluginCount++;
+
+					// Hook into Asset Registry Scan callbacks as as soon as it is loaded
+					if (ModuleName == TEXT("AssetRegistry"))
+					{
+						FAssetRegistryModule& AssetRegistryModule = FModuleManager::LoadModuleChecked<FAssetRegistryModule>(TEXT("AssetRegistry"));
+
+						AssetRegistryModule.Get().OnScanStarted().AddLambda([this]()
+							{
+								AssetRegistryScanStartTime = FDateTime::UtcNow();
+							});
+
+						AssetRegistryModule.Get().OnScanEnded().AddLambda([this]()
+							{
+								const float AssetRegistryScanTime = float((FDateTime::UtcNow() - AssetRegistryScanStartTime).GetTotalSeconds());
+								KPIRegistry.SetKPIValue(AssetRegistryScanKPI, AssetRegistryScanTime);
+							});
+					}
+
+					break;
+				}
+		
+				case EModuleChangeReason::ModuleUnloaded:
+				{
+					TotalPluginCount--;
+					break;
+				}	
+			}
+		});
 }
 
 extern ENGINE_API float GAverageFPS;
@@ -441,6 +489,17 @@ void FEditorPerformanceModule::UpdateKPIs(float InDeltaTime)
 		}
 	}
 
+	if (EditorState == EEditorState::PIE_Startup)
+	{
+		const float PIETransitionTime = float((FDateTime::UtcNow() - PIEStartTime).GetTotalSeconds());
+		KPIRegistry.SetKPIValue(IsFirstTimeToPIE? PIEFirstTransitionKPI : PIETransitionKPI, PIETransitionTime);
+	}
+
+	if (IsLoadingMap == true)
+	{
+		EditorLoadMapTime = float((FDateTime::UtcNow() - LoadMapStartTime).GetTotalSeconds());
+		KPIRegistry.SetKPIValue(EditorLoadMapKPI, EditorLoadMapTime);
+	}
 }
 
 bool FEditorPerformanceModule::IsHotLocalCacheCase() const
diff --git a/Engine/Plugins/Experimental/EditorPerformance/Source/EditorPerformance/Public/EditorPerformanceModule.h b/Engine/Plugins/Experimental/EditorPerformance/Source/EditorPerformance/Public/EditorPerformanceModule.h
--- a/Engine/Plugins/Experimental/EditorPerformance/Source/EditorPerformance/Public/EditorPerformanceModule.h
+++ b/Engine/Plugins/Experimental/EditorPerformance/Source/EditorPerformance/Public/EditorPerformanceModule.h
@@ -28,7 +28,7 @@ public:
 		Editor_Boot,
 		Editor_Initialize,
 		Editor_Interact,
-		PIE_Transtion,
+		PIE_Startup,
 		PIE_Interact,
 		PIE_Shutdown,
 	};
@@ -76,6 +76,9 @@ private:
 	FDateTime						LoadMapStartTime;
 	FDateTime						PIEStartTime;
 	FDateTime						PIEEndTime;
+	FDateTime						AssetRegistryScanStartTime;
+	bool							IsFirstTimeToPIE = true;
+	bool							IsLoadingMap = false;
 	EEditorState					EditorState = EEditorState::Editor_Boot;
 	float							BootToPIETime=0;
 	float							EditorBootTime = 0;
@@ -89,11 +92,13 @@ private:
 	uint32							HitchSampleCount = 0;
 	uint32							EditorHitchCount = 0;
 	uint32							PIEHitchCount = 0;
+	uint32							TotalPluginCount = 0;
 
 	FGuid							EditorBootKPI;
 	FGuid							EditorInitializeKPI;
 	FGuid							EditorLoadMapKPI;
 	FGuid							EditorHitchrateKPI;
+	FGuid							AssetRegistryScanKPI;
 	FGuid							TotalTimeToEditorKPI;
 	FGuid							TotalTimeToPIEKPI;
 	FGuid							PIEFirstTransitionKPI;
diff --git a/Engine/Plugins/Experimental/StudioTelemetry/Source/StudioTelemetry/Private/StudioTelemetryEditor.cpp b/Engine/Plugins/Experimental/StudioTelemetry/Source/StudioTelemetry/Private/StudioTelemetryEditor.cpp
--- a/Engine/Plugins/Experimental/StudioTelemetry/Source/StudioTelemetry/Private/StudioTelemetryEditor.cpp
+++ b/Engine/Plugins/Experimental/StudioTelemetry/Source/StudioTelemetry/Private/StudioTelemetryEditor.cpp
@@ -453,27 +453,6 @@ void FStudioTelemetryEditor::HeartbeatCallback()
 		HitchSampleCount = 0;
 		HitchAvergageFPS = 0;
 	}
-
-	// Monitor Asset Registry Scan
-	FAssetRegistryModule& AssetRegistryModule = FModuleManager::LoadModuleChecked<FAssetRegistryModule>(TEXT("AssetRegistry"));
-	const bool IsScaningAssets = AssetRegistryModule.Get().IsLoadingAssets() || AssetRegistryModule.Get().IsSearchAllAssets();
-
-	if (IsScaningAssets == false && AssetRegistryScanSpan.IsValid() == true)
-	{
-		// End the span
-		TArray<FAnalyticsEventAttribute> Attributes;
-		Attributes.Emplace(TEXT("MapName"), EditorMapName);
-
-		FStudioTelemetry::Get().EndSpan(AssetRegistryScanSpan, Attributes);
-
-		// No longer need the span for now so reset it
-		AssetRegistryScanSpan.Reset();
-	}
-	else if (IsScaningAssets == true && AssetRegistryScanSpan.IsValid() == false)
-	{
-		// Start the span
-		AssetRegistryScanSpan = FStudioTelemetry::Get().StartSpan(AssetRegistryScanSpanName);
-	}
 }
 
 void FStudioTelemetryEditor::Initialize()
@@ -567,8 +546,66 @@ void FStudioTelemetryEditor::Initialize()
 			FStudioTelemetryEditor::RecordEvent_CoreSystems(TEXT("LoadMap"), EditorLoadMapSpan->GetAttributes());
 		});
 
+	// Install any plugin load/unload callbacks
+	FModuleManager::Get().OnModulesChanged().AddLambda([this](FName ModuleName, EModuleChangeReason ChangeReason)
+		{
+			switch ( ChangeReason )
+			{ 
+				default:
+				{
+					break;
+				}
+				
+				case EModuleChangeReason::ModuleLoaded:
+				{
+					TotalPluginCount++;
+
+					// Hook into Asset Registry Scan callbacks as as soon as it is loaded
+					if (ModuleName == TEXT("AssetRegistry"))
+					{
+						FAssetRegistryModule& AssetRegistryModule = FModuleManager::LoadModuleChecked<FAssetRegistryModule>(TEXT("AssetRegistry"));
+
+						AssetRegistryModule.Get().OnScanStarted().AddLambda([this]()
+							{
+								// End the existing span if it overlaps
+								if (AssetRegistryScanSpan.IsValid())
+								{
+									TArray<FAnalyticsEventAttribute> Attributes;
+									Attributes.Emplace(TEXT("MapName"), EditorMapName);
+									AssetRegistryScanSpan.Reset();
+								}
+
+								// Start the Asset Registry Scan span
+								AssetRegistryScanSpan = FStudioTelemetry::Get().StartSpan(AssetRegistryScanSpanName);
+							});
+
+						AssetRegistryModule.Get().OnScanEnded().AddLambda([this]()
+							{
+								// End the Asset Registry Scan span
+								TArray<FAnalyticsEventAttribute> Attributes;
+								Attributes.Emplace(TEXT("MapName"), EditorMapName);
+								FStudioTelemetry::Get().EndSpan(AssetRegistryScanSpan, Attributes);
+								AssetRegistryScanSpan.Reset();
+							});
+					}
+
+					break;
+				}
+		
+				case EModuleChangeReason::ModuleUnloaded:
+				{
+					TotalPluginCount--;
+					break;
+				}	
+			}
+		});
+
 	FEditorDelegates::OnEditorBoot.AddLambda([this](double)
 		{	
+			TArray<FAnalyticsEventAttribute> Attributes;
+			Attributes.Emplace(TEXT("PluginCount"), TotalPluginCount);
+			EditorBootSpan->AddAttributes(Attributes);
+
 			FStudioTelemetry::Get().EndSpan(EditorBootSpan);
 
 			// Callback is received when the editor has booted but has not been initialized
@@ -585,7 +622,8 @@ void FStudioTelemetryEditor::Initialize()
 			// Editor has initialized
 			TArray<FAnalyticsEventAttribute> Attributes;
 			Attributes.Emplace(TEXT("MapName"), EditorMapName);
-			
+			Attributes.Emplace(TEXT("PluginCount"), TotalPluginCount);
+
 			// Editor has finished initializing so start the Editor Interact span
 			FStudioTelemetry::Get().EndSpan(EditorInitilizeSpan);
 			EditorInteractSpan = FStudioTelemetry::Get().StartSpan(EditorInteractSpanName, EditorSpan);
diff --git a/Engine/Plugins/Experimental/StudioTelemetry/Source/StudioTelemetry/Private/StudioTelemetryEditor.h b/Engine/Plugins/Experimental/StudioTelemetry/Source/StudioTelemetry/Private/StudioTelemetryEditor.h
--- a/Engine/Plugins/Experimental/StudioTelemetry/Source/StudioTelemetry/Private/StudioTelemetryEditor.h
+++ b/Engine/Plugins/Experimental/StudioTelemetry/Source/StudioTelemetry/Private/StudioTelemetryEditor.h
@@ -58,6 +58,7 @@ private:
 	const FName EditorInitilizeSpanName = TEXT("Editor.Initialize");
 	const FName EditorInteractSpanName = TEXT("Editor.Interact");
 	const FName EditorLoadMapSpanName = TEXT("Editor.LoadMap");
+	const FName AssetRegistryScanSpanName = TEXT("Editor.AssetRegistryScan");
 	const FName PIESpanName = TEXT("PIE");
 	const FName PIEStartupSpanName = TEXT("PIE.Startup");
 	const FName PIEPreBeginSpanName = TEXT("PIE.PreBegin");
@@ -67,7 +68,6 @@ private:
 	const FName CookingSpanName = TEXT("Cooking");
 	const FName HitchingSpanName = TEXT("Hitching");
 	const FName OpenAssetEditorSpan = TEXT("Open Asset Editor");
-	const FName AssetRegistryScanSpanName = TEXT("Asset Registry Scan");
 	const float HeartbeatIntervalSeconds = 1.0;
 	const float HitchSamplerIntervalSeconds = 0.1;
 	const float MinFPSForHitching = 5.0;
@@ -85,6 +85,7 @@ private:
 	double TimeToBootEditor;
 	double HitchAvergageFPS = 0;
 	uint32 HitchSampleCount = 0;
+	uint32 TotalPluginCount = 0;
 };
 
 #endif // WITH_EDITOR
diff --git a/Engine/Source/Developer/DerivedDataCache/Private/DerivedDataCache.cpp b/Engine/Source/Developer/DerivedDataCache/Private/DerivedDataCache.cpp
--- a/Engine/Source/Developer/DerivedDataCache/Private/DerivedDataCache.cpp
+++ b/Engine/Source/Developer/DerivedDataCache/Private/DerivedDataCache.cpp
@@ -106,8 +106,7 @@ namespace UE::DerivedData::CookStats
 				LocalGetMisses += UsageStats.GetStats.GetAccumulatedValueAnyThread(FCookStats::CallStats::EHitOrMiss::Miss, FCookStats::CallStats::EStatType::Counter);
 				LocalSpeedStats = (*LocalNode)->SpeedStats;
 			}
-			const int64 LocalGetTotal = LocalGetHits + LocalGetMisses;
-
+			
 			int64 ZenLocalGetHits = 0;
 			int64 ZenLocalGetMisses = 0;
 			FDerivedDataCacheSpeedStats ZenLocalSpeedStats;
@@ -117,9 +116,15 @@ namespace UE::DerivedData::CookStats
 				ZenLocalGetHits += UsageStats.GetStats.GetAccumulatedValueAnyThread(FCookStats::CallStats::EHitOrMiss::Hit, FCookStats::CallStats::EStatType::Counter);
 				ZenLocalGetMisses += UsageStats.GetStats.GetAccumulatedValueAnyThread(FCookStats::CallStats::EHitOrMiss::Miss, FCookStats::CallStats::EStatType::Counter);
 				ZenLocalSpeedStats = (*ZenLocalNode)->SpeedStats;
+
+				LocalGetHits = ZenLocalGetHits;
+				LocalGetMisses = ZenLocalGetMisses;
+				LocalSpeedStats = ZenLocalSpeedStats;
 			}			
-			const int64 ZenLocalGetTotal = ZenLocalGetHits + ZenLocalGetMisses;
-			
+
+			const int64 ZenLocalGetTotal = ZenLocalGetHits + ZenLocalGetMisses;	
+			const int64 LocalGetTotal = LocalGetHits + LocalGetMisses;
+
 			int64 SharedGetHits = 0;
 			int64 SharedGetMisses = 0;
 			FDerivedDataCacheSpeedStats SharedSpeedStats;
diff --git a/Engine/Source/Runtime/AssetRegistry/Private/AssetRegistry.cpp b/Engine/Source/Runtime/AssetRegistry/Private/AssetRegistry.cpp
--- a/Engine/Source/Runtime/AssetRegistry/Private/AssetRegistry.cpp
+++ b/Engine/Source/Runtime/AssetRegistry/Private/AssetRegistry.cpp
@@ -1904,6 +1904,8 @@ void FAssetRegistryImpl::SearchAllAssets(Impl::FEventContext& EventContext,
 		return;
 	}
 
+	EventContext.bScanStartedEventBroadcast = true;
+
 	FAssetDataGatherer& Gatherer = *GlobalGatherer;
 	if (!Gatherer.IsAsyncEnabled())
 	{
@@ -4884,6 +4886,7 @@ void FAssetRegistryImpl::OnInitialSearchCompleted(Impl::FEventContext& EventCont
 	GlobalGatherer->OnInitialSearchCompleted();
 
 	EventContext.bFileLoadedEventBroadcast = true;
+	EventContext.bScanEndedEventBroadcast = true;
 }
 
 void FAssetRegistryImpl::LogSearchDiagnostics(double StartTime)
@@ -8664,6 +8667,20 @@ void UAssetRegistryImpl::Broadcast(UE::AssetRegistry::Impl::FEventContext& Event
 		FileLoadedEvent.Broadcast();
 		EventContext.bFileLoadedEventBroadcast = false;
 	}
+
+	if (EventContext.bScanStartedEventBroadcast)
+	{
+		// Raise event when the scan is started
+		ScanStartedEvent.Broadcast();
+		EventContext.bScanStartedEventBroadcast = false;
+	}
+
+	if (EventContext.bScanEndedEventBroadcast)
+	{
+		// Raise event when the scan is ended
+		ScanEndedEvent.Broadcast();
+		EventContext.bScanEndedEventBroadcast = false;
+	}
 }
 
 
@@ -8767,6 +8784,16 @@ UAssetRegistryImpl::FFileLoadProgressUpdatedEvent& UAssetRegistryImpl::OnFileLoa
 	return FileLoadProgressUpdatedEvent;
 }
 
+UAssetRegistryImpl::FScanStartedEvent& UAssetRegistryImpl::OnScanStarted()
+{
+	return ScanStartedEvent;
+}
+
+UAssetRegistryImpl::FScanEndedEvent& UAssetRegistryImpl::OnScanEnded()
+{
+	return ScanEndedEvent;
+}
+
 namespace UE::AssetRegistry
 {
 const FAssetData* GetMostImportantAsset(TConstArrayView<const FAssetData*> PackageAssetDatas, EGetMostImportantAssetFlags InFlags)
diff --git a/Engine/Source/Runtime/AssetRegistry/Private/AssetRegistry.h b/Engine/Source/Runtime/AssetRegistry/Private/AssetRegistry.h
--- a/Engine/Source/Runtime/AssetRegistry/Private/AssetRegistry.h
+++ b/Engine/Source/Runtime/AssetRegistry/Private/AssetRegistry.h
@@ -130,6 +130,12 @@ public:
 	DECLARE_DERIVED_EVENT( UAssetRegistryImpl, IAssetRegistry::FPathRemovedEvent, FPathRemovedEvent);
 	virtual FPathRemovedEvent& OnPathRemoved() override;
 
+	DECLARE_DERIVED_EVENT(UAssetRegistryImpl, IAssetRegistry::FScanStartedEvent, FScanStartedEvent);
+	virtual FScanStartedEvent& OnScanStarted() override;
+
+	DECLARE_DERIVED_EVENT(UAssetRegistryImpl, IAssetRegistry::FScanEndedEvent, FScanEndedEvent);
+	virtual FScanEndedEvent& OnScanEnded() override;
+
 	virtual void AssetCreated(UObject* NewAsset) override;
 	virtual void AssetDeleted(UObject* DeletedAsset) override;
 	virtual void AssetRenamed(const UObject* RenamedAsset, const FString& OldObjectPath) override;
@@ -385,6 +391,12 @@ private:
 	/** The delegate to execute while loading files to update progress */
 	FFileLoadProgressUpdatedEvent FileLoadProgressUpdatedEvent;
 
+	/** The delegate to execute scanning has begun */
+	FScanStartedEvent ScanStartedEvent;
+
+	/** The delegate to execute scanning has ended */
+	FScanEndedEvent ScanEndedEvent;
+
 	UE::AssetRegistry::Impl::FEventContext DeferredEvents;
 
 	friend class UE::AssetRegistry::FAssetRegistryImpl;
diff --git a/Engine/Source/Runtime/AssetRegistry/Private/AssetRegistryImpl.h b/Engine/Source/Runtime/AssetRegistry/Private/AssetRegistryImpl.h
--- a/Engine/Source/Runtime/AssetRegistry/Private/AssetRegistryImpl.h
+++ b/Engine/Source/Runtime/AssetRegistry/Private/AssetRegistryImpl.h
@@ -657,6 +657,8 @@ struct FEventContext
 	TArray<FString> RequiredLoads;
 	TArray<FString> BlockedFiles;
 	bool bFileLoadedEventBroadcast = false;
+	bool bScanStartedEventBroadcast = false;
+	bool bScanEndedEventBroadcast = false;
 
 	/** Remove all stored events */
 	void Clear();
diff --git a/Engine/Source/Runtime/AssetRegistry/Public/AssetRegistry/IAssetRegistry.h b/Engine/Source/Runtime/AssetRegistry/Public/AssetRegistry/IAssetRegistry.h
--- a/Engine/Source/Runtime/AssetRegistry/Public/AssetRegistry/IAssetRegistry.h
+++ b/Engine/Source/Runtime/AssetRegistry/Public/AssetRegistry/IAssetRegistry.h
@@ -767,6 +767,12 @@ public:
 	DECLARE_TS_MULTICAST_DELEGATE( FFilesLoadedEvent );
 	virtual FFilesLoadedEvent& OnFilesLoaded() = 0;
 
+	DECLARE_TS_MULTICAST_DELEGATE(FScanStartedEvent);
+	virtual FScanStartedEvent& OnScanStarted() = 0;
+
+	DECLARE_TS_MULTICAST_DELEGATE(FScanEndedEvent);
+	virtual FScanEndedEvent& OnScanEnded() = 0;
+
 	/** Payload data for a file progress update */
 	struct FFileLoadProgressUpdateData
 	{
