Fix for Asset Registry Span tracking in Editor Performance Module and Studio Telemetry

diff --git a/Engine/Plugins/Experimental/EditorPerformance/Source/EditorPerformance/Private/EditorPerformanceModule.cpp b/Engine/Plugins/Experimental/EditorPerformance/Source/EditorPerformance/Private/EditorPerformanceModule.cpp
--- a/Engine/Plugins/Experimental/EditorPerformance/Source/EditorPerformance/Private/EditorPerformanceModule.cpp
+++ b/Engine/Plugins/Experimental/EditorPerformance/Source/EditorPerformance/Private/EditorPerformanceModule.cpp
@@ -37,6 +37,7 @@ const FName EditorInitializeKPIName = TEXT("Initialize");
 const FName EditorLoadMapKPIName = TEXT("Load Map");
 const FName EditorHitchrateKPIName = TEXT("Hitch Rate");
 const FName EditorAssetRegistryScanKPIName = TEXT("Asset Registry Scan");
+const FName EditorPluginCountKPIName = TEXT("Plugin Count");
 const FName TotalTimeToEditorKPIName = TEXT("Total Time To Editor");
 const FName TotalTimeToPIEKPIName = TEXT("Total Time To PIE");
 const FName PIEFirstTransitionKPIName = TEXT("First Transition");
@@ -55,7 +56,8 @@ float EditorBootKPILimit = 100;
 float EditorInitializeKPILimit = 160;
 float EditorLoadMapKPILimit = 120;
 float EditorHitchrateKPILimit = 25;
-float EditorAssetRegistryScanKPILimit = 60;
+float EditorAssetRegistryScanKPILimit = 140;
+float EditorPluginCountKPILimit = 1500;
 float TotalTimeToEditorKPILimit = 160;
 float PIEFirstTransitionKPILimit = 220;
 float PIETransitionKPILimit = 40;
@@ -190,25 +192,27 @@ void FEditorPerformanceModule::ShowPerformanceReportTab()
 void FEditorPerformanceModule::InitializeKPIs()
 {
 	// Declare the KPIs 
-	EditorBootKPI			= KPIRegistry.DeclareKPIValue(EditorCategoryName, EditorBootKPIName, 0.0, EditorBootKPILimit, FKPIValue::LessThan, FKPIValue::Minutes);
-	EditorInitializeKPI		= //KPIRegistry.DeclareKPIValue(EditorCategoryName, EditorInitializeKPIName, 0.0, EditorInitializeKPILimit, FKPIValue::LessThan, FKPIValue::Minutes);
-	EditorLoadMapKPI		= KPIRegistry.DeclareKPIValue(EditorCategoryName, EditorLoadMapKPIName, 0.0, EditorLoadMapKPILimit, FKPIValue::LessThan, FKPIValue::Minutes);
-	AssetRegistryScanKPI	= KPIRegistry.DeclareKPIValue(EditorCategoryName, EditorAssetRegistryScanKPIName, 0.0, EditorAssetRegistryScanKPILimit, FKPIValue::LessThan, FKPIValue::Minutes);
-	TotalTimeToEditorKPI	= KPIRegistry.DeclareKPIValue(EditorCategoryName, TotalTimeToEditorKPIName, 0.0, TotalTimeToEditorKPILimit, FKPIValue::LessThan, FKPIValue::Minutes);
-	PIEFirstTransitionKPI	= KPIRegistry.DeclareKPIValue(PIECategoryName, PIEFirstTransitionKPIName, 0.0, PIEFirstTransitionKPILimit, FKPIValue::LessThan, FKPIValue::Minutes);
-	PIETransitionKPI		= KPIRegistry.DeclareKPIValue(PIECategoryName, PIETransitionKPIName, 0.0, PIETransitionKPILimit, FKPIValue::LessThan, FKPIValue::Minutes);
-	PIEShutdownKPI			= KPIRegistry.DeclareKPIValue(PIECategoryName, PIEShutdownKPIName, 0.0, PIEShutdownKPILimit, FKPIValue::LessThan, FKPIValue::Minutes);
-	TotalTimeToPIEKPI		= KPIRegistry.DeclareKPIValue(PIECategoryName, TotalTimeToPIEKPIName, 0.0, TotalTimeToPIEKPILimit, FKPIValue::LessThan, FKPIValue::Minutes);
-	TotalDDCEfficiencyKPI	= KPIRegistry.DeclareKPIValue(CacheCategoryName, TotalDDCEfficiencyKPIName, 100.0, TotalDDCEffciencyKPILimit, FKPIValue::GreaterThan, FKPIValue::Percent);
-	LocalDDCEfficiencyKPI	= KPIRegistry.DeclareKPIValue(CacheCategoryName, LocalDDCEfficiencyKPIName, 100.0, LocalDDCEffciencyKPILimit, FKPIValue::GreaterThan, FKPIValue::Percent);
-	CoreCountKPI			= KPIRegistry.DeclareKPIValue(HardwareCategoryName, CoreCountKPIName, 128.0, CoreCountKPILimit, FKPIValue::GreaterThanOrEqual, FKPIValue::Decimal);
-	TotalMemoryKPI			= KPIRegistry.DeclareKPIValue(HardwareCategoryName, TotalMemoryKPIName, 128.0, TotalMemoryKPILimit, FKPIValue::GreaterThanOrEqual, FKPIValue::GigaBytes);
-	AvailableMemoryKPI		= KPIRegistry.DeclareKPIValue(HardwareCategoryName, AvailableMemoryKPIName, 128.0, AvailableMemoryKPILimit, FKPIValue::GreaterThanOrEqual, FKPIValue::GigaBytes);
+	EditorBootKPI				= KPIRegistry.DeclareKPIValue(EditorCategoryName, EditorBootKPIName, 0.0, EditorBootKPILimit, FKPIValue::LessThan, FKPIValue::Minutes);
+	EditorInitializeKPI			= //KPIRegistry.DeclareKPIValue(EditorCategoryName, EditorInitializeKPIName, 0.0, EditorInitializeKPILimit, FKPIValue::LessThan, FKPIValue::Minutes);
+	EditorLoadMapKPI			= KPIRegistry.DeclareKPIValue(EditorCategoryName, EditorLoadMapKPIName, 0.0, EditorLoadMapKPILimit, FKPIValue::LessThan, FKPIValue::Minutes);
+	EditorAssetRegistryScanKPI	= KPIRegistry.DeclareKPIValue(EditorCategoryName, EditorAssetRegistryScanKPIName, 0.0, EditorAssetRegistryScanKPILimit, FKPIValue::LessThan, FKPIValue::Minutes);
+	EditorPluginCountKPI		= KPIRegistry.DeclareKPIValue(EditorCategoryName, EditorPluginCountKPIName, 0.0, EditorPluginCountKPILimit, FKPIValue::LessThan, FKPIValue::Decimal);
+	TotalTimeToEditorKPI		= KPIRegistry.DeclareKPIValue(EditorCategoryName, TotalTimeToEditorKPIName, 0.0, TotalTimeToEditorKPILimit, FKPIValue::LessThan, FKPIValue::Minutes);
+	PIEFirstTransitionKPI		= KPIRegistry.DeclareKPIValue(PIECategoryName, PIEFirstTransitionKPIName, 0.0, PIEFirstTransitionKPILimit, FKPIValue::LessThan, FKPIValue::Minutes);
+	PIETransitionKPI			= KPIRegistry.DeclareKPIValue(PIECategoryName, PIETransitionKPIName, 0.0, PIETransitionKPILimit, FKPIValue::LessThan, FKPIValue::Minutes);
+	PIEShutdownKPI				= KPIRegistry.DeclareKPIValue(PIECategoryName, PIEShutdownKPIName, 0.0, PIEShutdownKPILimit, FKPIValue::LessThan, FKPIValue::Minutes);
+	TotalTimeToPIEKPI			= KPIRegistry.DeclareKPIValue(PIECategoryName, TotalTimeToPIEKPIName, 0.0, TotalTimeToPIEKPILimit, FKPIValue::LessThan, FKPIValue::Minutes);
+	TotalDDCEfficiencyKPI		= KPIRegistry.DeclareKPIValue(CacheCategoryName, TotalDDCEfficiencyKPIName, 100.0, TotalDDCEffciencyKPILimit, FKPIValue::GreaterThan, FKPIValue::Percent);
+	LocalDDCEfficiencyKPI		= KPIRegistry.DeclareKPIValue(CacheCategoryName, LocalDDCEfficiencyKPIName, 100.0, LocalDDCEffciencyKPILimit, FKPIValue::GreaterThan, FKPIValue::Percent);
+	CoreCountKPI				= KPIRegistry.DeclareKPIValue(HardwareCategoryName, CoreCountKPIName, 128.0, CoreCountKPILimit, FKPIValue::GreaterThanOrEqual, FKPIValue::Decimal);
+	TotalMemoryKPI				= KPIRegistry.DeclareKPIValue(HardwareCategoryName, TotalMemoryKPIName, 128.0, TotalMemoryKPILimit, FKPIValue::GreaterThanOrEqual, FKPIValue::GigaBytes);
+	AvailableMemoryKPI			= KPIRegistry.DeclareKPIValue(HardwareCategoryName, AvailableMemoryKPIName, 128.0, AvailableMemoryKPILimit, FKPIValue::GreaterThanOrEqual, FKPIValue::GigaBytes);
 
 	// Declare the KPI Hints
 	KPIRegistry.DeclareKPIHint(EditorBootKPI, LOCTEXT("EditorBootHintMessage", "The Editor boot time is slow.\nCheck you have enbabled a Game Feature Plugin profile for your project and that the expected local cache efficiency is met.\nIf you are booting the Editor in the background then disable the Use Less CPU in Background option in the settings."), LOCTEXT("EditorBootHintURL","https://docs.unrealengine.com/5.0/en-US/"));
-	KPIRegistry.DeclareKPIHint(TotalTimeToEditorKPI, LOCTEXT("EditorStartupHintMessage", "The Editor start - up time is slow.\nCheck you have enbabled a Game Feature Plugin profile for your project and that the expected local cache efficiency is met.\nIf you are booting the Editor in the background then disable the Use Less CPU in Background option in the settings."), LOCTEXT("EditorBootHintURL","https://docs.unrealengine.com/5.0/en-US/"));
-	
+	KPIRegistry.DeclareKPIHint(TotalTimeToEditorKPI, LOCTEXT("EditorStartupHintMessage", "The Editor start-up time is slow.\nCheck you have enbabled a Game Feature Plugin profile for your project and that the expected local cache efficiency is met.\nIf you are booting the Editor in the background then disable the Use Less CPU in Background option in the settings."), LOCTEXT("EditorBootHintURL","https://docs.unrealengine.com/5.0/en-US/"));
+	KPIRegistry.DeclareKPIHint(EditorPluginCountKPI, LOCTEXT("EditorPluginHintMessage", "The Editor is loading more plugins than expected and this will affect Editor start-up performance.\nCheck you have enbabled a Game Feature Plugin profile for your project."), LOCTEXT("EditorPluginHintURL", "https://docs.unrealengine.com/5.0/en-US/"));
+
 	KPIRegistry.DeclareKPIHint(PIETransitionKPI, LOCTEXT("PIETransitionHintMessage", "The Editor transtion to PIE is slow.\nCheck that the expected local cache efficiency is met.\nIf you are transitioning to PIE with the Editor in the background then disable the Use Less CPU in Background option in the settings."), LOCTEXT("PIETransitionHintURL", "https://docs.unrealengine.com/5.0/en-US/"));
 	
 	KPIRegistry.DeclareKPIHint(LocalDDCEfficiencyKPI, LOCTEXT("LocalCacheEfficencyHintMessage", "The Editor will not perform well if the local cache efficiency has not yet met the expected value.\nIf this is the first time you have booted the Editor after a sync then this is to be expected."), LOCTEXT("EditorCacheHintURL","https://docs.unrealengine.com/5.3/en-US/derived-data-cache/"));
@@ -350,13 +354,22 @@ void FEditorPerformanceModule::InitializeKPIs()
 
 						AssetRegistryModule.Get().OnScanStarted().AddLambda([this]()
 							{
-								AssetRegistryScanStartTime = FDateTime::UtcNow();
+								if (EditorAssetRegistryScanCount == 0)
+								{
+									AssetRegistryScanStartTime = FDateTime::UtcNow();
+								}	
+
+								EditorAssetRegistryScanCount++;
 							});
 
 						AssetRegistryModule.Get().OnScanEnded().AddLambda([this]()
 							{
-								const float AssetRegistryScanTime = float((FDateTime::UtcNow() - AssetRegistryScanStartTime).GetTotalSeconds());
-								KPIRegistry.SetKPIValue(AssetRegistryScanKPI, AssetRegistryScanTime);
+								EditorAssetRegistryScanCount--;
+
+								if (EditorAssetRegistryScanCount == 0)
+								{
+									EditorAssetRegistryScanTime = float((FDateTime::UtcNow() - AssetRegistryScanStartTime).GetTotalSeconds());
+								}	
 							});
 					}
 
@@ -402,7 +415,17 @@ void FEditorPerformanceModule::UpdateKPIs(float InDeltaTime)
 {
 	// Gather live hardware stats
 	KPIRegistry.SetKPIValue(AvailableMemoryKPI, static_cast<float>(FPlatformMemory::GetStats().AvailablePhysical)/(1024.0f * 1024.0f * 1024.0f));
-	
+
+	// Update stats that may have been captures before initialization
+	KPIRegistry.SetKPIValue(EditorPluginCountKPI, (float)TotalPluginCount);
+
+	if (EditorAssetRegistryScanCount!=0)
+	{
+		EditorAssetRegistryScanTime = float((FDateTime::UtcNow() - AssetRegistryScanStartTime).GetTotalSeconds());
+	}
+
+	KPIRegistry.SetKPIValue(EditorAssetRegistryScanKPI, EditorAssetRegistryScanTime);
+
 	// Gather the DDC summary stats
 	FDerivedDataCacheSummaryStats SummaryStats;
 	GatherDerivedDataCacheSummaryStats(SummaryStats);
diff --git a/Engine/Plugins/Experimental/EditorPerformance/Source/EditorPerformance/Public/EditorPerformanceModule.h b/Engine/Plugins/Experimental/EditorPerformance/Source/EditorPerformance/Public/EditorPerformanceModule.h
--- a/Engine/Plugins/Experimental/EditorPerformance/Source/EditorPerformance/Public/EditorPerformanceModule.h
+++ b/Engine/Plugins/Experimental/EditorPerformance/Source/EditorPerformance/Public/EditorPerformanceModule.h
@@ -84,6 +84,8 @@ private:
 	float							EditorBootTime = 0;
 	float							EditorStartUpTime = 0;
 	float							EditorLoadMapTime = 0;
+	float							EditorAssetRegistryScanTime = 0;
+	uint32							EditorAssetRegistryScanCount = 0;
 	FString							EditorMapName=TEXT("Boot");
 	FTimerHandle					HitchSamplerTimerHandle;
 	const float						HitchSamplerIntervalSeconds = 0.1f;
@@ -98,7 +100,8 @@ private:
 	FGuid							EditorInitializeKPI;
 	FGuid							EditorLoadMapKPI;
 	FGuid							EditorHitchrateKPI;
-	FGuid							AssetRegistryScanKPI;
+	FGuid							EditorAssetRegistryScanKPI;
+	FGuid							EditorPluginCountKPI;
 	FGuid							TotalTimeToEditorKPI;
 	FGuid							TotalTimeToPIEKPI;
 	FGuid							PIEFirstTransitionKPI;
diff --git a/Engine/Plugins/Experimental/StudioTelemetry/Source/StudioTelemetry/Private/StudioTelemetry.cpp b/Engine/Plugins/Experimental/StudioTelemetry/Source/StudioTelemetry/Private/StudioTelemetry.cpp
--- a/Engine/Plugins/Experimental/StudioTelemetry/Source/StudioTelemetry/Private/StudioTelemetry.cpp
+++ b/Engine/Plugins/Experimental/StudioTelemetry/Source/StudioTelemetry/Private/StudioTelemetry.cpp
@@ -314,3 +314,9 @@ TSharedPtr<IAnalyticsSpan> FStudioTelemetry::GetSpan(const FName Name)
 	return AnalyticsTracer.IsValid() ? AnalyticsTracer->GetSpan(Name) : TSharedPtr<IAnalyticsSpan>();
 }
 
+TSharedPtr<IAnalyticsSpan> FStudioTelemetry::GetSessionSpan() const
+{
+	return AnalyticsTracer.IsValid() ? AnalyticsTracer->GetSessionSpan() : TSharedPtr<IAnalyticsSpan>();
+}
+
+
diff --git a/Engine/Plugins/Experimental/StudioTelemetry/Source/StudioTelemetry/Private/StudioTelemetryEditor.cpp b/Engine/Plugins/Experimental/StudioTelemetry/Source/StudioTelemetry/Private/StudioTelemetryEditor.cpp
--- a/Engine/Plugins/Experimental/StudioTelemetry/Source/StudioTelemetry/Private/StudioTelemetryEditor.cpp
+++ b/Engine/Plugins/Experimental/StudioTelemetry/Source/StudioTelemetry/Private/StudioTelemetryEditor.cpp
@@ -567,25 +567,26 @@ void FStudioTelemetryEditor::Initialize()
 
 						AssetRegistryModule.Get().OnScanStarted().AddLambda([this]()
 							{
-								// End the existing span if it overlaps
-								if (AssetRegistryScanSpan.IsValid())
+								if (AssetRegistryScanCount == 0)
 								{
-									TArray<FAnalyticsEventAttribute> Attributes;
-									Attributes.Emplace(TEXT("MapName"), EditorMapName);
-									AssetRegistryScanSpan.Reset();
-								}
+									// Start the Asset Registry Scan span
+									AssetRegistryScanSpan = FStudioTelemetry::Get().StartSpan(AssetRegistryScanSpanName, EditorSpan);
+								}	
 
-								// Start the Asset Registry Scan span
-								AssetRegistryScanSpan = FStudioTelemetry::Get().StartSpan(AssetRegistryScanSpanName);
+								AssetRegistryScanCount++;
 							});
 
 						AssetRegistryModule.Get().OnScanEnded().AddLambda([this]()
 							{
-								// End the Asset Registry Scan span
-								TArray<FAnalyticsEventAttribute> Attributes;
-								Attributes.Emplace(TEXT("MapName"), EditorMapName);
-								FStudioTelemetry::Get().EndSpan(AssetRegistryScanSpan, Attributes);
-								AssetRegistryScanSpan.Reset();
+								AssetRegistryScanCount--;
+
+								if (AssetRegistryScanCount == 0)
+								{
+									// End the Asset Registry Scan span
+									TArray<FAnalyticsEventAttribute> Attributes;
+									Attributes.Emplace(TEXT("MapName"), EditorMapName);
+									FStudioTelemetry::Get().EndSpan(AssetRegistryScanSpan, Attributes);
+								}	
 							});
 					}
 
diff --git a/Engine/Plugins/Experimental/StudioTelemetry/Source/StudioTelemetry/Private/StudioTelemetryEditor.h b/Engine/Plugins/Experimental/StudioTelemetry/Source/StudioTelemetry/Private/StudioTelemetryEditor.h
--- a/Engine/Plugins/Experimental/StudioTelemetry/Source/StudioTelemetry/Private/StudioTelemetryEditor.h
+++ b/Engine/Plugins/Experimental/StudioTelemetry/Source/StudioTelemetry/Private/StudioTelemetryEditor.h
@@ -86,6 +86,7 @@ private:
 	double HitchAvergageFPS = 0;
 	uint32 HitchSampleCount = 0;
 	uint32 TotalPluginCount = 0;
+	uint32 AssetRegistryScanCount = 0;
 };
 
 #endif // WITH_EDITOR
diff --git a/Engine/Plugins/Experimental/StudioTelemetry/Source/StudioTelemetry/Public/StudioTelemetry.h b/Engine/Plugins/Experimental/StudioTelemetry/Source/StudioTelemetry/Public/StudioTelemetry.h
--- a/Engine/Plugins/Experimental/StudioTelemetry/Source/StudioTelemetry/Public/StudioTelemetry.h
+++ b/Engine/Plugins/Experimental/StudioTelemetry/Source/StudioTelemetry/Public/StudioTelemetry.h
@@ -67,6 +67,9 @@ public:
 	/** Get an active span by name, non active spans will not be available*/
 	STUDIOTELEMETRY_API TSharedPtr<IAnalyticsSpan> GetSpan(const FName Name);
 
+	/** Get the root session span*/
+	STUDIOTELEMETRY_API TSharedPtr<IAnalyticsSpan> GetSessionSpan() const;
+
 	/** Callback for interception of telemetry events recording that can be used by Developers to send telemetry events to their own back end, though it is recommended that Developers implement their own IAnalyticsProvider via their own IAnalyticsProviderModule*/
 	STUDIOTELEMETRY_API void SetRecordEventCallback(OnRecordEvent);
 
