Fix warning V637: Two opposite conditions were encountered. The second condition is always false.
#rnx

diff --git a/Engine/Plugins/Experimental/EditorPerformance/Source/EditorPerformance/Private/SEditorPerformanceStatusBar.cpp b/Engine/Plugins/Experimental/EditorPerformance/Source/EditorPerformance/Private/SEditorPerformanceStatusBar.cpp
--- a/Engine/Plugins/Experimental/EditorPerformance/Source/EditorPerformance/Private/SEditorPerformanceStatusBar.cpp
+++ b/Engine/Plugins/Experimental/EditorPerformance/Source/EditorPerformance/Private/SEditorPerformanceStatusBar.cpp
@@ -233,7 +233,7 @@ EActiveTimerReturnType SEditorPerformanceStatusBarWidget::UpdateState(double InC
 
 	if (WarningCount > 0)
 	{
-		if (WarningCount == 0)
+		if (WarningCount == 1)
 		{
 			EditorPerformanceStateMessage = FText::FromString(*FString::Printf(TEXT("There Is 1 Warning.\n\nView Performance Report For Details.")));
 		}
