Latest updates to Editor Performance tool
Added Virtual Assets Efficiency KPI
Added Stall Rate KPI
Added failure count for KPIs and added this to the UI
Insights Snapshots will generate multiple captures for each failure within the session upto a a maximum of 10 per session
Insights capture and Telemetry event generation movede from StatusBar to EditorPerformanceModule
Fixed issue where Editor Load Map tiem was incorrect when booting without a Map
Split Hitch and Stall rate KPIs bewteen Editor and PIE scopes
Added LocalGetMisses to DDC cache statistics
Turned off notifications by default, users should opt-in

diff --git a/Engine/Plugins/Experimental/EditorPerformance/Source/EditorPerformance/EditorPerformance.Build.cs b/Engine/Plugins/Experimental/EditorPerformance/Source/EditorPerformance/EditorPerformance.Build.cs
--- a/Engine/Plugins/Experimental/EditorPerformance/Source/EditorPerformance/EditorPerformance.Build.cs
+++ b/Engine/Plugins/Experimental/EditorPerformance/Source/EditorPerformance/EditorPerformance.Build.cs
@@ -26,6 +26,7 @@ public class EditorPerformance : ModuleRules
 				"MessageLog",
 				"ToolWidgets",
 				"DerivedDataCache",
+				"Virtualization",
 				"StudioTelemetry"
 			});
 
diff --git a/Engine/Plugins/Experimental/EditorPerformance/Source/EditorPerformance/Private/EditorPerformanceModule.cpp b/Engine/Plugins/Experimental/EditorPerformance/Source/EditorPerformance/Private/EditorPerformanceModule.cpp
--- a/Engine/Plugins/Experimental/EditorPerformance/Source/EditorPerformance/Private/EditorPerformanceModule.cpp
+++ b/Engine/Plugins/Experimental/EditorPerformance/Source/EditorPerformance/Private/EditorPerformanceModule.cpp
@@ -18,9 +18,11 @@
 #include "Editor/EditorPerformanceSettings.h"
 #include "AssetRegistry/AssetRegistryModule.h"
 #include "DerivedDataCacheUsageStats.h"
+#include "Virtualization/VirtualizationSystem.h"
 #include "Trace/Trace.h"
 #include "StudioTelemetry.h"
 #include "HAL/PlatformFileManager.h"
+#include "ProfilingDebugging/StallDetector.h"
 
 #define LOCTEXT_NAMESPACE "EditorPerformance"
  
@@ -35,7 +37,8 @@ const FName HardwareCategoryName = TEXT("Hardware");
 const FName EditorBootKPIName = TEXT("Boot");
 const FName EditorInitializeKPIName = TEXT("Initialize");
 const FName EditorLoadMapKPIName = TEXT("Load Map");
-const FName EditorHitchrateKPIName = TEXT("Hitch Rate");
+const FName EditorHitchRateKPIName = TEXT("Hitch Rate");
+const FName EditorStallRateKPIName = TEXT("Stall Rate");
 const FName EditorAssetRegistryScanKPIName = TEXT("Asset Registry Scan");
 const FName EditorPluginCountKPIName = TEXT("Plugin Count");
 const FName TotalTimeToEditorKPIName = TEXT("Total Time To Editor");
@@ -43,11 +46,13 @@ const FName TotalTimeToPIEKPIName = TEXT("Total Time To PIE");
 const FName PIEFirstTransitionKPIName = TEXT("First Transition");
 const FName PIETransitionKPIName = TEXT("Iterative Transition");
 const FName PIEShutdownKPIName = TEXT("Shutdown");
-const FName PIEHitchrateKPIName = TEXT("Hitch Rate");
+const FName PIEHitchRateKPIName = TEXT("Hitch Rate");
+const FName PIEStallRateKPIName = TEXT("Stall Rate");
 const FName CloudDDCLatencyKPIName = TEXT("Unreal Cloud DDC Latency");
 const FName CloudDDCReadSpeedKPIName = TEXT("Unreal Cloud DDC Speed");
 const FName TotalDDCEfficiencyKPIName = TEXT("Effective Efficiency");
 const FName LocalDDCEfficiencyKPIName = TEXT("Local Efficiency");
+const FName VirtualAssetEfficiencyKPIName = TEXT("Virtual Asset Efficiency");
 const FName CoreCountKPIName = TEXT("Core Count");
 const FName TotalMemoryKPIName = TEXT("Total Memory");
 const FName AvailableMemoryKPIName = TEXT("Available Memory");
@@ -55,19 +60,22 @@ const FName AvailableMemoryKPIName = TEXT("Available Memory");
 float EditorBootKPILimit = 100;
 float EditorInitializeKPILimit = 160;
 float EditorLoadMapKPILimit = 120;
-float EditorHitchrateKPILimit = 25;
+float EditorHitchRateKPILimit = 25;
+float EditorStallRateKPILimit = 25;
 float EditorAssetRegistryScanKPILimit = 140;
 float EditorPluginCountKPILimit = 1500;
 float TotalTimeToEditorKPILimit = 160;
 float PIEFirstTransitionKPILimit = 220;
 float PIETransitionKPILimit = 40;
 float PIEShutdownKPILimit = 10;
-float PIEHitchrateKPILimit = 25;
+float PIEHitchRateKPILimit = 25;
+float PIEStallRateKPILimit = 25;
 float TotalTimeToPIEKPILimit = 600;
 float CloudDDCLatencyKPILimit = 100;
 float CloudDDCReadSpeedKPILimit = 10;
 float TotalDDCEffciencyKPILimit = 90;
 float LocalDDCEffciencyKPILimit = 85;
+float VirtualAssetEfficiencyKPILimit = 95;
 float CoreCountKPILimit = 32;
 float TotalMemoryKPILimit = 64;
 float AvailableMemoryKPILimit = 16;
@@ -82,11 +90,11 @@ void FEditorPerformanceModule::StartupModule()
 
 void FEditorPerformanceModule::ShutdownModule()
 {
-	TerminateUI();
+	TerminateEditor();
 	TerminateKPIs();
 }
 
-void FEditorPerformanceModule::InitializeUI()
+void FEditorPerformanceModule::InitializeEditor()
 {
 	UEditorPerformanceSettings* EditorPerformanceSettings = GetMutableDefault<UEditorPerformanceSettings>();
 
@@ -101,10 +109,11 @@ void FEditorPerformanceModule::InitializeUI()
 		// Enable any experimental features
 		if (EditorPerformanceSettings->bEnableExperimentalFeatures)
 		{
-			CloudDDCLatencyKPI		= KPIRegistry.DeclareKPIValue(CacheCategoryName, CloudDDCLatencyKPIName, 0.0, CloudDDCLatencyKPILimit, FKPIValue::LessThan, FKPIValue::Milliseconds);
-			CloudDDCReadSpeedKPI	= KPIRegistry.DeclareKPIValue(CacheCategoryName, CloudDDCReadSpeedKPIName, 100.0, CloudDDCReadSpeedKPILimit, FKPIValue::GreaterThan, FKPIValue::MegaBitsPerSecond);
-			EditorHitchrateKPI		= KPIRegistry.DeclareKPIValue(EditorCategoryName, EditorHitchrateKPIName, 0.0, EditorHitchrateKPILimit, FKPIValue::LessThan, FKPIValue::Percent);
-			PIEHitchrateKPI			= KPIRegistry.DeclareKPIValue(PIECategoryName, PIEHitchrateKPIName, 0.0, PIEHitchrateKPILimit, FKPIValue::LessThan, FKPIValue::Percent);
+			CloudDDCLatencyKPI			= KPIRegistry.DeclareKPIValue(CacheCategoryName, CloudDDCLatencyKPIName, 0.0, CloudDDCLatencyKPILimit, FKPIValue::LessThan, FKPIValue::Milliseconds);
+			CloudDDCReadSpeedKPI		= KPIRegistry.DeclareKPIValue(CacheCategoryName, CloudDDCReadSpeedKPIName, 100.0, CloudDDCReadSpeedKPILimit, FKPIValue::GreaterThan, FKPIValue::MegaBitsPerSecond);
+			EditorStallRateKPI			= KPIRegistry.DeclareKPIValue(EditorCategoryName, EditorStallRateKPIName, 0.0, EditorStallRateKPILimit, FKPIValue::LessThan, FKPIValue::Percent);
+			PIEStallRateKPI				= KPIRegistry.DeclareKPIValue(PIECategoryName, PIEStallRateKPIName, 0.0, PIEStallRateKPILimit, FKPIValue::LessThan, FKPIValue::Percent);
+			VirtualAssetEfficiencyKPI	= KPIRegistry.DeclareKPIValue(CacheCategoryName, VirtualAssetEfficiencyKPIName, 100.0, VirtualAssetEfficiencyKPILimit, FKPIValue::GreaterThan, FKPIValue::Percent);
 		}
 
 		// Populate the notification list with all KPI values if it is empty. 
@@ -120,10 +129,14 @@ void FEditorPerformanceModule::InitializeUI()
 		}
 	}
 
+	FTimerDelegate HeartBeatDelegate;
+	HeartBeatDelegate.BindRaw(this, &FEditorPerformanceModule::HeartBeatCallback);
+	GEditor->GetTimerManager()->SetTimer(HeartBeatTimerHandle, HeartBeatDelegate, HeartBeatIntervalSeconds, true);
+	
 	FTimerDelegate HitchSamplerDelegate;
 	HitchSamplerDelegate.BindRaw(this, &FEditorPerformanceModule::HitchSamplerCallback);
 	GEditor->GetTimerManager()->SetTimer(HitchSamplerTimerHandle, HitchSamplerDelegate, HitchSamplerIntervalSeconds, true);
-	
+
 	const FSlateIcon PerformanceReportIcon(FAppStyle::GetAppStyleSetName(), "EditorPerformance.Report.Panel");
 
 	FGlobalTabmanager::Get()->RegisterNomadTabSpawner(EditorPerformanceReportTabName, FOnSpawnTab::CreateRaw(this, &FEditorPerformanceModule::CreatePerformanceReportTab))
@@ -149,7 +162,7 @@ void FEditorPerformanceModule::InitializeUI()
 	EditorPerfSection.AddEntry(FToolMenuEntry::InitWidget("EditorPerformanceStatusBar", CreateStatusBarWidget(), FText::GetEmpty(), true, false));
 }
 
-void FEditorPerformanceModule::TerminateUI()
+void FEditorPerformanceModule::TerminateEditor()
 {
 	if (FSlateApplication::IsInitialized())
 	{
@@ -197,17 +210,19 @@ void FEditorPerformanceModule::InitializeKPIs()
 	EditorLoadMapKPI			= KPIRegistry.DeclareKPIValue(EditorCategoryName, EditorLoadMapKPIName, 0.0, EditorLoadMapKPILimit, FKPIValue::LessThan, FKPIValue::Minutes);
 	EditorAssetRegistryScanKPI	= KPIRegistry.DeclareKPIValue(EditorCategoryName, EditorAssetRegistryScanKPIName, 0.0, EditorAssetRegistryScanKPILimit, FKPIValue::LessThan, FKPIValue::Minutes);
 	EditorPluginCountKPI		= KPIRegistry.DeclareKPIValue(EditorCategoryName, EditorPluginCountKPIName, 0.0, EditorPluginCountKPILimit, FKPIValue::LessThan, FKPIValue::Decimal);
+	EditorHitchRateKPI			= KPIRegistry.DeclareKPIValue(EditorCategoryName, EditorHitchRateKPIName, 0.0, EditorHitchRateKPILimit, FKPIValue::LessThan, FKPIValue::Percent);
 	TotalTimeToEditorKPI		= KPIRegistry.DeclareKPIValue(EditorCategoryName, TotalTimeToEditorKPIName, 0.0, TotalTimeToEditorKPILimit, FKPIValue::LessThan, FKPIValue::Minutes);
 	PIEFirstTransitionKPI		= KPIRegistry.DeclareKPIValue(PIECategoryName, PIEFirstTransitionKPIName, 0.0, PIEFirstTransitionKPILimit, FKPIValue::LessThan, FKPIValue::Minutes);
 	PIETransitionKPI			= KPIRegistry.DeclareKPIValue(PIECategoryName, PIETransitionKPIName, 0.0, PIETransitionKPILimit, FKPIValue::LessThan, FKPIValue::Minutes);
 	PIEShutdownKPI				= KPIRegistry.DeclareKPIValue(PIECategoryName, PIEShutdownKPIName, 0.0, PIEShutdownKPILimit, FKPIValue::LessThan, FKPIValue::Minutes);
+	PIEHitchRateKPI				= KPIRegistry.DeclareKPIValue(PIECategoryName, PIEHitchRateKPIName, 0.0, PIEHitchRateKPILimit, FKPIValue::LessThan, FKPIValue::Percent);
 	TotalTimeToPIEKPI			= KPIRegistry.DeclareKPIValue(PIECategoryName, TotalTimeToPIEKPIName, 0.0, TotalTimeToPIEKPILimit, FKPIValue::LessThan, FKPIValue::Minutes);
-	TotalDDCEfficiencyKPI		= KPIRegistry.DeclareKPIValue(CacheCategoryName, TotalDDCEfficiencyKPIName, 100.0, TotalDDCEffciencyKPILimit, FKPIValue::GreaterThan, FKPIValue::Percent);
+	//TotalDDCEfficiencyKPI		= KPIRegistry.DeclareKPIValue(CacheCategoryName, TotalDDCEfficiencyKPIName, 100.0, TotalDDCEffciencyKPILimit, FKPIValue::GreaterThan, FKPIValue::Percent);
 	LocalDDCEfficiencyKPI		= KPIRegistry.DeclareKPIValue(CacheCategoryName, LocalDDCEfficiencyKPIName, 100.0, LocalDDCEffciencyKPILimit, FKPIValue::GreaterThan, FKPIValue::Percent);
 	CoreCountKPI				= KPIRegistry.DeclareKPIValue(HardwareCategoryName, CoreCountKPIName, 128.0, CoreCountKPILimit, FKPIValue::GreaterThanOrEqual, FKPIValue::Decimal);
 	TotalMemoryKPI				= KPIRegistry.DeclareKPIValue(HardwareCategoryName, TotalMemoryKPIName, 128.0, TotalMemoryKPILimit, FKPIValue::GreaterThanOrEqual, FKPIValue::GigaBytes);
 	AvailableMemoryKPI			= KPIRegistry.DeclareKPIValue(HardwareCategoryName, AvailableMemoryKPIName, 128.0, AvailableMemoryKPILimit, FKPIValue::GreaterThanOrEqual, FKPIValue::GigaBytes);
-
+	
 	// Declare the KPI Hints
 	KPIRegistry.DeclareKPIHint(EditorBootKPI, LOCTEXT("EditorBootHintMessage", "The Editor boot time is slow.\nCheck you have enbabled a Game Feature Plugin profile for your project and that the expected local cache efficiency is met.\nIf you are booting the Editor in the background then disable the Use Less CPU in Background option in the settings."), LOCTEXT("EditorBootHintURL","https://docs.unrealengine.com/5.0/en-US/"));
 	KPIRegistry.DeclareKPIHint(TotalTimeToEditorKPI, LOCTEXT("EditorStartupHintMessage", "The Editor start-up time is slow.\nCheck you have enbabled a Game Feature Plugin profile for your project and that the expected local cache efficiency is met.\nIf you are booting the Editor in the background then disable the Use Less CPU in Background option in the settings."), LOCTEXT("EditorBootHintURL","https://docs.unrealengine.com/5.0/en-US/"));
@@ -257,10 +272,9 @@ void FEditorPerformanceModule::InitializeKPIs()
 			KPIRegistry.SetKPIValue(EditorInitializeKPI, EditorStartUpTime-EditorBootTime );
 			KPIRegistry.SetKPIValue(TotalTimeToEditorKPI, EditorStartUpTime);
 
-			InitializeUI();
+			InitializeEditor();
 
 			EditorState = EEditorState::Editor_Interact;
-			EditorHitchCount = 0;
 		});
 
 	FEditorDelegates::OnMapLoad.AddLambda([this](const FString& MapName, FCanLoadMap& OutCanLoadMap)
@@ -273,20 +287,23 @@ void FEditorPerformanceModule::InitializeKPIs()
 		{
 			IsLoadingMap = false;
 
-			EditorMapName = FPaths::GetBaseFilename(MapName);
-	
-			EditorLoadMapTime = float((FDateTime::UtcNow() - LoadMapStartTime).GetTotalSeconds());
-			KPIRegistry.SetKPIValue(EditorLoadMapKPI, EditorLoadMapTime);
-	
-			// Apply any profile that matches the currently loaded map
-			for (FKPIProfiles::TConstIterator It(KPIRegistry.GetKPIProfiles()); It; ++It)
+			if (MapName.Len() > 0)
 			{
-				const FKPIProfile& Profile = It->Value;
+				EditorMapName = FPaths::GetBaseFilename(MapName);
+
+				EditorLoadMapTime = float((FDateTime::UtcNow() - LoadMapStartTime).GetTotalSeconds());
+				KPIRegistry.SetKPIValue(EditorLoadMapKPI, EditorLoadMapTime);
 
-				if (Profile.MapName == EditorMapName)
+				// Apply any profile that matches the currently loaded map
+				for (FKPIProfiles::TConstIterator It(KPIRegistry.GetKPIProfiles()); It; ++It)
 				{
-					KPIProfileName = It->Key;
-					KPIRegistry.ApplyKPIProfile(It->Value);
+					const FKPIProfile& Profile = It->Value;
+
+					if (Profile.MapName == EditorMapName)
+					{
+						KPIProfileName = It->Key;
+						KPIRegistry.ApplyKPIProfile(It->Value);
+					}
 				}
 			}
 		});
@@ -314,8 +331,6 @@ void FEditorPerformanceModule::InitializeKPIs()
 			}
 
 			EditorState = EEditorState::PIE_Interact;
-			PIEHitchCount = 0;
-
 			IsFirstTimeToPIE = false;
 		});
 
@@ -329,9 +344,7 @@ void FEditorPerformanceModule::InitializeKPIs()
 		{
 			const float PIEShutdownTime = float((FDateTime::UtcNow() - PIEEndTime).GetTotalSeconds());
 			KPIRegistry.SetKPIValue(PIEShutdownKPI, PIEShutdownTime);
-			
 			EditorState = EEditorState::Editor_Interact;
-			EditorHitchCount = 0;
 		});
 
 	FModuleManager::Get().OnModulesChanged().AddLambda([this](FName ModuleName, EModuleChangeReason ChangeReason)
@@ -359,12 +372,12 @@ void FEditorPerformanceModule::InitializeKPIs()
 									AssetRegistryScanStartTime = FDateTime::UtcNow();
 								}	
 
-								EditorAssetRegistryScanCount++;
+								FPlatformAtomics::InterlockedIncrement(&EditorAssetRegistryScanCount);
 							});
 
 						AssetRegistryModule.Get().OnScanEnded().AddLambda([this]()
 							{
-								EditorAssetRegistryScanCount--;
+								FPlatformAtomics::InterlockedDecrement(&EditorAssetRegistryScanCount);
 
 								if (EditorAssetRegistryScanCount == 0)
 								{
@@ -383,44 +396,78 @@ void FEditorPerformanceModule::InitializeKPIs()
 				}	
 			}
 		});
+
+#if STALL_DETECTOR
+
+	UE::FStallDetector::StallDetected.AddLambda([this](const UE::FStallDetectedParams& Params)
+		{
+			FPlatformAtomics::InterlockedIncrement(&StallDetectedCount);
+		});
+
+	UE::FStallDetector::StallCompleted.AddLambda([this](const UE::FStallCompletedParams& Params)
+		{
+			FPlatformAtomics::InterlockedDecrement(&StallDetectedCount);
+		});
+
+#endif //STALL_DETECTOR
 }
 
 extern ENGINE_API float GAverageFPS;
 
 void FEditorPerformanceModule::HitchSamplerCallback()
 {
-	// Only sample framerate when we have focus
+	static uint32 SampleCount = 0;
+	static uint32 HitchCount = 0;
+	static uint32 StallCount = 0;
+	
+	// Only sample framerate and hitches and stalls when we have focus
 	if (FApp::HasFocus())
 	{
-		// Sample a rolling average of FPS 
-		HitchAvergageFPS = (HitchAvergageFPS * HitchSampleCount + GAverageFPS) / (double)(HitchSampleCount + 1);
-		HitchSampleCount++;
-
-		if (HitchAvergageFPS < MinFPSForHitching)
+		if (GAverageFPS < MinFPSForHitching)
 		{
-			if (EditorState == EEditorState::Editor_Interact)
-			{
-				EditorHitchCount++;
-			}
-
-			if (EditorState == EEditorState::PIE_Interact)
-			{
-				PIEHitchCount++;
-			}
+			// This sample was hitching
+			HitchCount++;
 		}
 	}
+
+	if (StallDetectedCount > 0)
+	{
+		// This sample was stalling
+		StallCount++;
+	}
+
+	SampleCount++;
+
+	// Update HitchRate and StallRate after a set number of samples
+	if (SampleCount > MinSamplesForHitching)
+	{
+		HitchRate = 100.0f * (float)HitchCount / (float)SampleCount;
+		HitchCount = 0;
+
+		StallRate = 100.0f * (float)StallCount / (float)SampleCount;
+		StallCount = 0;
+		
+		SampleCount = 0;
+	}
+}
+
+void FEditorPerformanceModule::HeartBeatCallback()
+{
+	// Update the KPIS
+	UpdateKPIs(HeartBeatIntervalSeconds);
 }
 
 void FEditorPerformanceModule::UpdateKPIs(float InDeltaTime)
 {
 	// Gather live hardware stats
-	KPIRegistry.SetKPIValue(AvailableMemoryKPI, static_cast<float>(FPlatformMemory::GetStats().AvailablePhysical)/(1024.0f * 1024.0f * 1024.0f));
+	KPIRegistry.SetKPIValue(AvailableMemoryKPI, static_cast<float>(FPlatformMemory::GetStats().AvailablePhysical) / (1024.0f * 1024.0f * 1024.0f));
 
 	// Update stats that may have been captures before initialization
 	KPIRegistry.SetKPIValue(EditorPluginCountKPI, (float)TotalPluginCount);
 
-	if (EditorAssetRegistryScanCount!=0)
+	if (EditorAssetRegistryScanCount > 0)
 	{
+		// Keep track of the first Asset Registry scan time
 		EditorAssetRegistryScanTime = float((FDateTime::UtcNow() - AssetRegistryScanStartTime).GetTotalSeconds());
 	}
 
@@ -431,6 +478,7 @@ void FEditorPerformanceModule::UpdateKPIs(float InDeltaTime)
 	GatherDerivedDataCacheSummaryStats(SummaryStats);
 
 	int64 TotalCloudGetHits = 0;
+	int64 LocalGetHitsMisses = 0;
 	float CloudLatency = 0.0;
 	float CloudReadSpeed = 0.0;
 
@@ -466,6 +514,47 @@ void FEditorPerformanceModule::UpdateKPIs(float InDeltaTime)
 				KPIRegistry.SetKPIValue(LocalDDCEfficiencyKPI, Value);
 			}
 		}
+		else if (Stat.Key == TEXT("LocalGetMisses"))
+		{
+			LocalGetHitsMisses = FCString::Atoi(*Stat.Value);
+		}
+	}
+
+	// Gather the Virtual Assets stats
+	if (UE::Virtualization::IVirtualizationSystem::Get().IsEnabled())
+	{
+		TArray<UE::Virtualization::FBackendStats> BackEndStatsList = UE::Virtualization::IVirtualizationSystem::Get().GetBackendStatistics();
+
+		int64 CacheBackEndPullCount = 0;
+		int64 PersistentBackEndPullCount = 0;
+
+		for (const UE::Virtualization::FBackendStats& BackEndStats : BackEndStatsList)
+		{
+			switch (BackEndStats.Type)
+			{
+				case UE::Virtualization::EStorageType::Persistent:
+				{
+					PersistentBackEndPullCount += BackEndStats.PayloadActivity.Pull.PayloadCount;
+					break;
+				}
+
+				default:
+				case UE::Virtualization::EStorageType::Cache:
+				{
+					CacheBackEndPullCount += BackEndStats.PayloadActivity.Pull.PayloadCount;
+					break;
+				}
+			}
+		}
+
+		const int64 TotalBackEndPullCount = CacheBackEndPullCount+PersistentBackEndPullCount;
+
+		if (TotalBackEndPullCount > 0)
+		{
+			// Gather Virtualization analytics
+			const float VirtualAssetsEfficiency = 100.0f * float(CacheBackEndPullCount) / float(TotalBackEndPullCount);
+			KPIRegistry.SetKPIValue(VirtualAssetEfficiencyKPI, VirtualAssetsEfficiency);
+		}
 	}
 
 	// Evaluate Cloud Cache performance
@@ -475,9 +564,9 @@ void FEditorPerformanceModule::UpdateKPIs(float InDeltaTime)
 	static float AverageCloudLatency = 0;
 	static float AverageCloudReadSpeed = 0;
 
-	ElapsedCloudCacheHits = TotalCloudGetHits- PreviousTotalCloudGetHits;
+	ElapsedCloudCacheHits = TotalCloudGetHits - PreviousTotalCloudGetHits;
 	PreviousTotalCloudGetHits = TotalCloudGetHits;
-	
+
 	if (TotalCloudGetHits < MinimalCloudGetHits)
 	{
 		KPIRegistry.InvalidateKPIValue(CloudDDCLatencyKPI);
@@ -489,39 +578,57 @@ void FEditorPerformanceModule::UpdateKPIs(float InDeltaTime)
 		KPIRegistry.SetKPIValue(CloudDDCReadSpeedKPI, CloudReadSpeed);
 	}
 
-	// Evaluate Hitching
+	// Record Hitch Rate
 	if (EditorState == EEditorState::Editor_Interact)
 	{
-		if (InDeltaTime > 0.0f)
-		{
-			KPIRegistry.SetKPIValue(EditorHitchrateKPI, 100.0f * (float)EditorHitchCount* HitchSamplerIntervalSeconds);
-			
-			EditorHitchCount = 0;
-			HitchSampleCount = 0;
-		}
+		KPIRegistry.SetKPIValue(EditorHitchRateKPI, HitchRate);
+		KPIRegistry.SetKPIValue(EditorStallRateKPI, StallRate);
 	}
 
 	if (EditorState == EEditorState::PIE_Interact)
 	{
-		if (InDeltaTime > 0.0f )
-		{
-			KPIRegistry.SetKPIValue(PIEHitchrateKPI, 100.0f * (float)PIEHitchCount * HitchSamplerIntervalSeconds);
-
-			PIEHitchCount = 0;
-			HitchSampleCount = 0;
-		}
+		KPIRegistry.SetKPIValue(PIEHitchRateKPI, HitchRate);
+		KPIRegistry.SetKPIValue(PIEStallRateKPI, StallRate);
 	}
 
-	if (EditorState == EEditorState::PIE_Startup)
+	static TArray<FGuid> RecordedKPIEvent;
+	
+	// Check for KPIs that have exceeded their value
+	for (FKPIValues::TConstIterator It(GetKPIRegistry().GetKPIValues()); It; ++It)
 	{
-		const float PIETransitionTime = float((FDateTime::UtcNow() - PIEStartTime).GetTotalSeconds());
-		KPIRegistry.SetKPIValue(IsFirstTimeToPIE? PIEFirstTransitionKPI : PIETransitionKPI, PIETransitionTime);
-	}
+		const FKPIValue& KPIValue = It->Value;
 
-	if (IsLoadingMap == true)
-	{
-		EditorLoadMapTime = float((FDateTime::UtcNow() - LoadMapStartTime).GetTotalSeconds());
-		KPIRegistry.SetKPIValue(EditorLoadMapKPI, EditorLoadMapTime);
+		if (KPIValue.GetState() == FKPIValue::Bad)
+		{
+			if (RecordedKPIEvent.Find(KPIValue.Id) == INDEX_NONE)
+			{
+				// KPI has exceeding the threshold for the first time	
+				const UEditorPerformanceSettings* EditorPerformanceSettings = GetDefault<UEditorPerformanceSettings>();
+
+				if (EditorPerformanceSettings)
+				{
+					if (EditorPerformanceSettings->bEnableSnapshots)
+					{
+						// Create an Insights Snapshot
+						RecordInsightsSnaphshot(KPIValue);
+					}
+
+					if (EditorPerformanceSettings->bEnableTelemetry)
+					{
+						// Create a new telemetry event
+						RecordTelemetryEvent(KPIValue);
+					}
+				}
+
+				// Add this KPI to the list so we don't send the event again
+				RecordedKPIEvent.Emplace(KPIValue.Id);
+			}
+		}
+		else
+		{
+			// No longer exceeding threshold, so next time this KPI is exceeded we will record the event
+			RecordedKPIEvent.Remove(KPIValue.Id);
+		}
 	}
 }
 
@@ -539,8 +646,10 @@ bool FEditorPerformanceModule::IsHotLocalCacheCase() const
 
 bool FEditorPerformanceModule::RecordInsightsSnaphshot(const FKPIValue& KPIValue)
 {
+	const uint32 MaxKPITraceCount = 10;
+
 	// Create the snapshot file
-	FString FileName = KPIValue.Name.ToString() + TEXT(".utrace");
+	FString FileName = FString::Printf( TEXT("%s_%d.utrace"), *KPIValue.Path.ToString(), KPIValue.FailureCount % MaxKPITraceCount);
 	FString FolderPath = FPaths::ProjectSavedDir() / TEXT("EditorPerformance");
 
 	// Create the full output path
diff --git a/Engine/Plugins/Experimental/EditorPerformance/Source/EditorPerformance/Private/KPIValue.cpp b/Engine/Plugins/Experimental/EditorPerformance/Source/EditorPerformance/Private/KPIValue.cpp
--- a/Engine/Plugins/Experimental/EditorPerformance/Source/EditorPerformance/Private/KPIValue.cpp
+++ b/Engine/Plugins/Experimental/EditorPerformance/Source/EditorPerformance/Private/KPIValue.cpp
@@ -12,6 +12,8 @@ void FKPIValue::SetValue(float Value)
 {
 	CurrentValue = Value;
 
+	EState PreviousState = State;
+
 	switch (Compare)
 	{
 		default:
@@ -39,6 +41,11 @@ void FKPIValue::SetValue(float Value)
 			break;
 		}
 	}
+	
+	if (PreviousState != State && State == FKPIValue::Bad)
+	{
+		FailureCount++;
+	}
 }
 
 
diff --git a/Engine/Plugins/Experimental/EditorPerformance/Source/EditorPerformance/Private/SEditorPerformanceDialogs.cpp b/Engine/Plugins/Experimental/EditorPerformance/Source/EditorPerformance/Private/SEditorPerformanceDialogs.cpp
--- a/Engine/Plugins/Experimental/EditorPerformance/Source/EditorPerformance/Private/SEditorPerformanceDialogs.cpp
+++ b/Engine/Plugins/Experimental/EditorPerformance/Source/EditorPerformance/Private/SEditorPerformanceDialogs.cpp
@@ -265,6 +265,37 @@ TSharedRef<SWidget> SEditorPerformanceReportDialog::GetSettingsGridPanel()
 		];
 
 	Panel->AddSlot(4, Row)
+		.HAlign(HAlign_Left)
+		[
+			SNew(SCheckBox)
+				.IsChecked_Lambda([]
+					{
+						const UEditorPerformanceSettings* EditorPerformanceSettings = GetDefault<UEditorPerformanceSettings>();
+						return EditorPerformanceSettings && EditorPerformanceSettings->bShowFrameRateAndMemory ? ECheckBoxState::Checked : ECheckBoxState::Unchecked;
+					})
+				.OnCheckStateChanged_Lambda([this](ECheckBoxState NewState)
+					{
+						UEditorPerformanceSettings* EditorPerformanceSettings = GetMutableDefault<UEditorPerformanceSettings>();
+
+						if (EditorPerformanceSettings)
+						{
+							EditorPerformanceSettings->bShowFrameRateAndMemory = NewState == ECheckBoxState::Checked;
+							EditorPerformanceSettings->PostEditChange();
+							EditorPerformanceSettings->SaveConfig();
+						}
+
+						UpdateGridPanels(0.0f, 0.0f);
+					})
+						.Padding(FMargin(4.0f, 0.0f))
+						[
+							SNew(STextBlock)
+								.Text(LOCTEXT("EnableShowFrameRateAndMemoryText", "Diagnostics"))
+								.ToolTipText(LOCTEXT("EnableShowFrameRateAndMemoryToolTip", "Show the Frame Rate, Memory and Stalls."))
+								.ColorAndOpacity(EStyleColor::Foreground)
+						]
+		];
+
+	Panel->AddSlot(5, Row)
 		.HAlign(HAlign_Left)
 		.Padding(FMargin(10.0f, 10.0f))
 		[
@@ -493,9 +524,19 @@ TSharedRef<SWidget> SEditorPerformanceReportDialog::GetKPIGridPanel()
 			.Text(LOCTEXT("ExpectedValueColumn", "Expected"))
 		];
 
+	Panel->AddSlot(4, Row)
+		.HAlign(HAlign_Left)
+		[
+			SNew(STextBlock)
+				.Margin(DefaultMargin)
+				.ColorAndOpacity(EStyleColor::Foreground)
+				.Font(TitleFont)
+				.Text(LOCTEXT("ExpectedValueColumn", "Failures"))
+		];
+
 	if (EnableNotifcations)
 	{
-		Panel->AddSlot(5, Row)
+		Panel->AddSlot(6, Row)
 			.HAlign(HAlign_Center)
 			[
 				SNew(STextBlock)
@@ -597,9 +638,18 @@ TSharedRef<SWidget> SEditorPerformanceReportDialog::GetKPIGridPanel()
 						.Text(FText::FromString(*FKPIValue::GetValueAsString(KPIValue.ThresholdValue, KPIValue.DisplayType)))
 					];
 
+				Panel->AddSlot(4, Row)
+					.HAlign(HAlign_Left)
+					[
+						SNew(STextBlock)
+							.Margin(DefaultMargin)
+							.ColorAndOpacity(KPIColor)
+							.Text(FText::FromString(*FString::Printf(TEXT("%d"),KPIValue.FailureCount)))
+					];
+
 				if (KPIValueState != FKPIValue::Good)
 				{
-					Panel->AddSlot(4, Row)
+					Panel->AddSlot(5, Row)
 						[
 							SNew(SImage)
 							.Image(KPIWarningIcon)
@@ -620,7 +670,7 @@ TSharedRef<SWidget> SEditorPerformanceReportDialog::GetKPIGridPanel()
 
 			if (EnableNotifcations)
 			{
-				Panel->AddSlot(5, Row)
+				Panel->AddSlot(6, Row)
 					.HAlign(HAlign_Left)
 					[
 						SNew(SComboBox<FName>)
diff --git a/Engine/Plugins/Experimental/EditorPerformance/Source/EditorPerformance/Private/SEditorPerformanceStatusBar.cpp b/Engine/Plugins/Experimental/EditorPerformance/Source/EditorPerformance/Private/SEditorPerformanceStatusBar.cpp
--- a/Engine/Plugins/Experimental/EditorPerformance/Source/EditorPerformance/Private/SEditorPerformanceStatusBar.cpp
+++ b/Engine/Plugins/Experimental/EditorPerformance/Source/EditorPerformance/Private/SEditorPerformanceStatusBar.cpp
@@ -175,11 +175,10 @@ EActiveTimerReturnType SEditorPerformanceStatusBarWidget::UpdateState(double InC
 
 	const UEditorPerformanceSettings* EditorPerformanceSettings = GetDefault<UEditorPerformanceSettings>();
 
-	// Update the KPIS
-	EditorPerfModule.UpdateKPIs(InDeltaTime);
-
 	WarningCount = 0;
 
+	EditorPerfModule.UpdateKPIs(InDeltaTime);
+
 	// Check for KPIs that have exceeded their value
 	for (FKPIValues::TConstIterator It(EditorPerfModule.GetKPIRegistry().GetKPIValues()); It; ++It)
 	{
@@ -203,20 +202,6 @@ EActiveTimerReturnType SEditorPerformanceStatusBarWidget::UpdateState(double InC
 
 					CurrentNotificationName = KPIValue.Path;
 				}
-
-				if (RecordedSnapshot.Find(KPIValue.Path) == INDEX_NONE && EditorPerformanceSettings->bEnableSnapshots)
-				{
-					// Create an Insights Snapshot
-					EditorPerfModule.RecordInsightsSnaphshot(KPIValue);
-					RecordedSnapshot.Emplace(KPIValue.Path);
-				}
-
-				if (RecordedTelemetry.Find(KPIValue.Path) == INDEX_NONE && EditorPerformanceSettings->bEnableTelemetry)
-				{
-					// Create a new telemetry event
-					EditorPerfModule.RecordTelemetryEvent(KPIValue);
-					RecordedTelemetry.Emplace(KPIValue.Path);
-				}
 			}
 
 			WarningCount++;
@@ -226,8 +211,6 @@ EActiveTimerReturnType SEditorPerformanceStatusBarWidget::UpdateState(double InC
 			// No longer exceeding threshold, so no need to acknowledge the last time it was raised to the user
 			// There may be subsequent times that this same KPI is exceeded this session so we may want to alert the user again
 			AcknowledgedNotifications.Remove(KPIValue.Path);
-			RecordedSnapshot.Remove(KPIValue.Path);
-			RecordedTelemetry.Remove(KPIValue.Path);
 		}
 	}
 
diff --git a/Engine/Plugins/Experimental/EditorPerformance/Source/EditorPerformance/Private/SEditorPerformanceStatusBar.h b/Engine/Plugins/Experimental/EditorPerformance/Source/EditorPerformance/Private/SEditorPerformanceStatusBar.h
--- a/Engine/Plugins/Experimental/EditorPerformance/Source/EditorPerformance/Private/SEditorPerformanceStatusBar.h
+++ b/Engine/Plugins/Experimental/EditorPerformance/Source/EditorPerformance/Private/SEditorPerformanceStatusBar.h
@@ -66,8 +66,6 @@ private:
 	FText							CurrentNotificationMessage;
 	FName							CurrentNotificationName;
 	TArray<FName>					AcknowledgedNotifications;
-	TArray<FName>					RecordedSnapshot;
-	TArray<FName>					RecordedTelemetry;
 	uint32							WarningCount = 0;
 
 };
diff --git a/Engine/Plugins/Experimental/EditorPerformance/Source/EditorPerformance/Public/EditorPerformanceModule.h b/Engine/Plugins/Experimental/EditorPerformance/Source/EditorPerformance/Public/EditorPerformanceModule.h
--- a/Engine/Plugins/Experimental/EditorPerformance/Source/EditorPerformance/Public/EditorPerformanceModule.h
+++ b/Engine/Plugins/Experimental/EditorPerformance/Source/EditorPerformance/Public/EditorPerformanceModule.h
@@ -62,12 +62,13 @@ private:
 	TSharedRef<SDockTab> CreatePerformanceReportTab(const FSpawnTabArgs& Args);	
 	TWeakPtr<SDockTab> PerformanceReportTab;
 
-	void InitializeUI();
-	void TerminateUI();
+	void InitializeEditor();
+	void TerminateEditor();
 
 	void InitializeKPIs();
 	void TerminateKPIs();
 
+	void HeartBeatCallback();
 	void HitchSamplerCallback();
 
 	FKPIRegistry					KPIRegistry;
@@ -76,7 +77,7 @@ private:
 	FDateTime						LoadMapStartTime;
 	FDateTime						PIEStartTime;
 	FDateTime						PIEEndTime;
-	FDateTime						AssetRegistryScanStartTime;
+	FDateTime						AssetRegistryScanStartTime; 
 	bool							IsFirstTimeToPIE = true;
 	bool							IsLoadingMap = false;
 	EEditorState					EditorState = EEditorState::Editor_Boot;
@@ -85,21 +86,24 @@ private:
 	float							EditorStartUpTime = 0;
 	float							EditorLoadMapTime = 0;
 	float							EditorAssetRegistryScanTime = 0;
-	uint32							EditorAssetRegistryScanCount = 0;
+	volatile int32					EditorAssetRegistryScanCount = 0;
 	FString							EditorMapName=TEXT("Boot");
+	FTimerHandle					HeartBeatTimerHandle;
+	const float						HeartBeatIntervalSeconds = 1.0f;
 	FTimerHandle					HitchSamplerTimerHandle;
 	const float						HitchSamplerIntervalSeconds = 0.1f;
-	const float						MinFPSForHitching = 5.0f;
-	double							HitchAvergageFPS = 0;
-	uint32							HitchSampleCount = 0;
-	uint32							EditorHitchCount = 0;
-	uint32							PIEHitchCount = 0;
+	const float						MinFPSForHitching = 15.0f;
+	const uint32					MinSamplesForHitching = 10;
+	volatile int32					StallDetectedCount = 0;
+	float							HitchRate = 0;
+	float							StallRate = 0;
 	uint32							TotalPluginCount = 0;
 
 	FGuid							EditorBootKPI;
 	FGuid							EditorInitializeKPI;
 	FGuid							EditorLoadMapKPI;
-	FGuid							EditorHitchrateKPI;
+	FGuid							EditorHitchRateKPI;
+	FGuid							EditorStallRateKPI;
 	FGuid							EditorAssetRegistryScanKPI;
 	FGuid							EditorPluginCountKPI;
 	FGuid							TotalTimeToEditorKPI;
@@ -107,11 +111,13 @@ private:
 	FGuid							PIEFirstTransitionKPI;
 	FGuid							PIETransitionKPI;
 	FGuid							PIEShutdownKPI;
-	FGuid							PIEHitchrateKPI;
+	FGuid							PIEHitchRateKPI;
+	FGuid							PIEStallRateKPI;
 	FGuid							CloudDDCLatencyKPI;
 	FGuid							CloudDDCReadSpeedKPI;
 	FGuid							TotalDDCEfficiencyKPI;
 	FGuid							LocalDDCEfficiencyKPI;
+	FGuid							VirtualAssetEfficiencyKPI;
 	FGuid							CoreCountKPI;
 	FGuid							TotalMemoryKPI;
 	FGuid							AvailableMemoryKPI;
diff --git a/Engine/Plugins/Experimental/EditorPerformance/Source/EditorPerformance/Public/KPIValue.h b/Engine/Plugins/Experimental/EditorPerformance/Source/EditorPerformance/Public/KPIValue.h
--- a/Engine/Plugins/Experimental/EditorPerformance/Source/EditorPerformance/Public/KPIValue.h
+++ b/Engine/Plugins/Experimental/EditorPerformance/Source/EditorPerformance/Public/KPIValue.h
@@ -69,6 +69,7 @@ public:
 	FName			Path;
 	float			CurrentValue = 0;
 	float			ThresholdValue = 0;
+	uint32			FailureCount =0;
 	EState			State = EState::NotSet;
 	ECompare		Compare = ECompare::LessThan;
 	EDisplayType	DisplayType = EDisplayType::Number;
diff --git a/Engine/Plugins/Experimental/StudioTelemetry/Source/AnalyticsLog/Public/AnalyticsLog.h b/Engine/Plugins/Experimental/StudioTelemetry/Source/AnalyticsLog/Public/AnalyticsLog.h
--- a/Engine/Plugins/Experimental/StudioTelemetry/Source/AnalyticsLog/Public/AnalyticsLog.h
+++ b/Engine/Plugins/Experimental/StudioTelemetry/Source/AnalyticsLog/Public/AnalyticsLog.h
@@ -16,7 +16,6 @@
 #include "Templates/UnrealTemplate.h"
 
 class IAnalyticsProvider;
-class IAnalyticsProviderET;
 
 /**
  *  Public implementation of IAnalyticsProviderModule that returns a JSON formatted FAnalayticsProviderLog object
diff --git a/Engine/Plugins/Experimental/StudioTelemetry/Source/StudioTelemetry/Private/StudioTelemetryEditor.cpp b/Engine/Plugins/Experimental/StudioTelemetry/Source/StudioTelemetry/Private/StudioTelemetryEditor.cpp
--- a/Engine/Plugins/Experimental/StudioTelemetry/Source/StudioTelemetry/Private/StudioTelemetryEditor.cpp
+++ b/Engine/Plugins/Experimental/StudioTelemetry/Source/StudioTelemetry/Private/StudioTelemetryEditor.cpp
@@ -525,25 +525,33 @@ void FStudioTelemetryEditor::Initialize()
 
 	FEditorDelegates::OnMapLoad.AddLambda([this](const FString& MapName, FCanLoadMap& OutCanLoadMap)
 		{
-			// The Editor loads a new map
-			EditorLoadMapSpan = FStudioTelemetry::Get().StartSpan(EditorLoadMapSpanName, EditorSpan);
+			if (MapName.Len() > 0)
+			{
+				// The Editor loads a new map
+				EditorLoadMapSpan = FStudioTelemetry::Get().StartSpan(EditorLoadMapSpanName, EditorSpan);
+			}
 		});
 
 	FEditorDelegates::OnMapOpened.AddLambda([this](const FString& MapName, bool Unused)
 		{
-			// The new editor map was actually opened
-			EditorMapName = FPaths::GetBaseFilename(MapName);
+			if (EditorLoadMapSpan.IsValid())
+			{
+				// The new editor map was actually opened
+				EditorMapName = FPaths::GetBaseFilename(MapName);
 
-			TArray<FAnalyticsEventAttribute> Attributes;
-			Attributes.Emplace(TEXT("MapName"), EditorMapName);
+				TArray<FAnalyticsEventAttribute> Attributes;
+				Attributes.Emplace(TEXT("MapName"), EditorMapName);
 
-			EditorSpan->AddAttributes(Attributes);
-			EditorLoadMapSpan->AddAttributes(Attributes);
+				EditorSpan->AddAttributes(Attributes);
+				EditorLoadMapSpan->AddAttributes(Attributes);
 
-			FStudioTelemetry::Get().EndSpan(EditorLoadMapSpan);
+				FStudioTelemetry::Get().EndSpan(EditorLoadMapSpan);
+				
+				FStudioTelemetryEditor::RecordEvent_Loading(TEXT("LoadMap"), EditorLoadMapSpan->GetDuration(), EditorLoadMapSpan->GetAttributes());
+				FStudioTelemetryEditor::RecordEvent_CoreSystems(TEXT("LoadMap"), EditorLoadMapSpan->GetAttributes());
 
-			FStudioTelemetryEditor::RecordEvent_Loading(TEXT("LoadMap"), EditorLoadMapSpan->GetDuration(), EditorLoadMapSpan->GetAttributes());
-			FStudioTelemetryEditor::RecordEvent_CoreSystems(TEXT("LoadMap"), EditorLoadMapSpan->GetAttributes());
+				EditorLoadMapSpan.Reset();
+			}	
 		});
 
 	// Install any plugin load/unload callbacks
@@ -694,20 +702,14 @@ void FStudioTelemetryEditor::Initialize()
 			// PIE mode has been started. The user has pressed the Start PIE button.
 			// Finish the Editor span
 			FStudioTelemetry::Get().EndSpan(EditorSpan);
+			EditorSessionCount++;
 
 			// Start PIE span
 			PIESpan = FStudioTelemetry::Get().StartSpan(PIESpanName);
 
-			if (PIETransitionCount==0)
-			{ 
-				PIEStartupSpan = FStudioTelemetry::Get().StartSpan(PIEStartupSpanName, PIESpan);
-			}
-			else
-			{
-				// Append the PIE transition count to the PIE name
-				PIEStartupSpan = FStudioTelemetry::Get().StartSpan(FName(*FString::Printf(TEXT("%s%d"), *PIEStartupSpanName.ToString(), PIETransitionCount)), PIESpan);
-			}
-
+			// Append the PIE transition count to the PIE name
+			PIEStartupSpan = FStudioTelemetry::Get().StartSpan(PIESessionCount == 0? PIEStartupSpanName : FName(*FString::Printf(TEXT("%s%d"), *PIEStartupSpanName.ToString(), PIESessionCount)), PIESpan);
+		
 			TArray<FAnalyticsEventAttribute> Attributes;
 			Attributes.Emplace(TEXT("MapName"), EditorMapName);
 
@@ -717,7 +719,7 @@ void FStudioTelemetryEditor::Initialize()
 
 	FEditorDelegates::PreBeginPIE.AddLambda([this](bool)
 		{	
-			PIEPreBeginSpan = FStudioTelemetry::Get().StartSpan(PIEPreBeginSpanName, PIEStartupSpan);
+			PIEPreBeginSpan = FStudioTelemetry::Get().StartSpan(PIESessionCount == 0 ? PIEPreBeginSpanName : FName(*FString::Printf(TEXT("%s%d"), *PIEPreBeginSpanName.ToString(), PIESessionCount)), PIESpan);
 			PIEPreBeginSpan->AddAttributes(PIESpan->GetAttributes());
 		});
 
@@ -759,7 +761,7 @@ void FStudioTelemetryEditor::Initialize()
 					
 					// Keep track of the PIE transition counts
 					TArray<FAnalyticsEventAttribute> Attributes;
-					Attributes.Emplace(TEXT("PIE_TransitionCount"), PIETransitionCount );
+					Attributes.Emplace(TEXT("PIE_TransitionCount"), PIESessionCount );
 					
 					PIESpan->AddAttributes(Attributes);
 					PIEStartupSpan->AddAttributes(Attributes);
@@ -770,7 +772,7 @@ void FStudioTelemetryEditor::Initialize()
 					FStudioTelemetryEditor::RecordEvent_Loading(TEXT("PIE.TotalStartupTime"), PIEStartupSpan->GetDuration(), PIEStartupSpan->GetAttributes());
 					FStudioTelemetryEditor::RecordEvent_CoreSystems(TEXT("PIE.TotalStartupTime"), PIEStartupSpan->GetAttributes());
 
-					if (PIETransitionCount == 0)
+					if (PIESessionCount == 0)
 					{
 						const double TimeInEditor = EditorLoadMapSpan.IsValid() ? EditorLoadMapSpan->GetDuration() : 0.0;
 						const double TimeToStartPIE = PIEStartupSpan->GetDuration();
@@ -780,11 +782,9 @@ void FStudioTelemetryEditor::Initialize()
 						FStudioTelemetryEditor::RecordEvent_Loading(TEXT("TimeToPIE"), TimeToBootToPIE, PIEStartupSpan->GetAttributes());
 						FStudioTelemetryEditor::RecordEvent_CoreSystems(TEXT("TimeToPIE"), PIEStartupSpan->GetAttributes());
 					}
-
-					PIETransitionCount++;
 				}
 
-				PIEInteractSpan = FStudioTelemetry::Get().StartSpan(PIEInteractSpanName, PIESpan);
+				PIEInteractSpan = FStudioTelemetry::Get().StartSpan( PIESessionCount==0? PIEInteractSpanName : FName(*FString::Printf(TEXT("%s%d"), *PIEInteractSpanName.ToString(), PIESessionCount)), PIESpan);
 				PIEInteractSpan->AddAttributes(PIESpan->GetAttributes());
 			}
 		});
@@ -795,7 +795,7 @@ void FStudioTelemetryEditor::Initialize()
 			{
 				// PIE is ending so no longer interactive
 				FStudioTelemetry::Get().EndSpan(PIEInteractSpan);
-				PIEShutdownSpan = FStudioTelemetry::Get().StartSpan(PIEShutdownSpanName, PIESpan);
+				PIEShutdownSpan = FStudioTelemetry::Get().StartSpan(PIESessionCount == 0 ? PIEShutdownSpanName : FName(*FString::Printf(TEXT("%s%d"), *PIEShutdownSpanName.ToString(), PIESessionCount)), PIESpan);
 				PIEShutdownSpan->AddAttributes(PIESpan->GetAttributes());
 			}
 		});
@@ -811,12 +811,14 @@ void FStudioTelemetryEditor::Initialize()
 				FStudioTelemetryEditor::RecordEvent_CoreSystems(TEXT("PIE.EndTime"), PIESpan->GetAttributes());
 			}
 
+			PIESessionCount++;
+
 			TArray<FAnalyticsEventAttribute> Attributes;
 			Attributes.Emplace(TEXT("MapName"), EditorMapName);
 
 			// Restart the Editor span
 			EditorSpan = FStudioTelemetry::Get().StartSpan(EditorSpanName, Attributes);
-			EditorInteractSpan = FStudioTelemetry::Get().StartSpan(EditorInteractSpanName, EditorSpan);
+			EditorInteractSpan = FStudioTelemetry::Get().StartSpan(FName(*FString::Printf(TEXT("%s%d"), *EditorInteractSpanName.ToString(), EditorSessionCount)), EditorSpan);
 		});
 
 	
diff --git a/Engine/Plugins/Experimental/StudioTelemetry/Source/StudioTelemetry/Private/StudioTelemetryEditor.h b/Engine/Plugins/Experimental/StudioTelemetry/Source/StudioTelemetry/Private/StudioTelemetryEditor.h
--- a/Engine/Plugins/Experimental/StudioTelemetry/Source/StudioTelemetry/Private/StudioTelemetryEditor.h
+++ b/Engine/Plugins/Experimental/StudioTelemetry/Source/StudioTelemetry/Private/StudioTelemetryEditor.h
@@ -79,7 +79,8 @@ private:
 	FTimerHandle TelemetryHitchSamplerTimerHandle;
 	FString EditorMapName;
 	FString PIEMapName;
-	uint32 PIETransitionCount = 0;
+	uint32 EditorSessionCount = 0;
+	uint32 PIESessionCount = 0;
 	double SessionStartTime;
 	double AssetOpenStartTime;
 	double TimeToBootEditor;
diff --git a/Engine/Source/Developer/DerivedDataCache/Private/DerivedDataCache.cpp b/Engine/Source/Developer/DerivedDataCache/Private/DerivedDataCache.cpp
--- a/Engine/Source/Developer/DerivedDataCache/Private/DerivedDataCache.cpp
+++ b/Engine/Source/Developer/DerivedDataCache/Private/DerivedDataCache.cpp
@@ -189,6 +189,7 @@ namespace UE::DerivedData::CookStats
 				TEXT("TotalPutHitPct"), SafeDivide(RootPutHits, RootPutTotal),	
 				TEXT("PutMissPct"), SafeDivide(RootPutMisses, RootPutTotal),
 				TEXT("LocalGetHits"), LocalGetHits,
+				TEXT("LocalGetMisses"), LocalGetMisses,
 				TEXT("LocalGetTotal"), LocalGetTotal,
 				TEXT("LocalGetHitPct"), SafeDivide(LocalGetHits, LocalGetTotal),
 				TEXT("SharedGetHits"), SharedGetHits,
diff --git a/Engine/Source/Editor/UnrealEd/Classes/Editor/EditorPerformanceSettings.h b/Engine/Source/Editor/UnrealEd/Classes/Editor/EditorPerformanceSettings.h
--- a/Engine/Source/Editor/UnrealEd/Classes/Editor/EditorPerformanceSettings.h
+++ b/Engine/Source/Editor/UnrealEd/Classes/Editor/EditorPerformanceSettings.h
@@ -39,7 +39,7 @@ class UEditorPerformanceSettings : public UDeveloperSettings
 
 	/** Whether to enable Editor Perf Warnings */
 	UPROPERTY(EditAnywhere, config, Category = EditorPerformanceTool, meta = (DisplayName = "Enable Notifcations", ConfigRestartRequired = false, EditCondition = "bEnableEditorPeformanceTool"))
-	bool bEnableNotifications = true;
+	bool bEnableNotifications = false;
 
 	/** Specific measurements to notify */
 	UPROPERTY(EditAnywhere, config, Category = EditorPerformanceTool, meta = (DisplayName = "Notification List", ConfigRestartRequired = false, EditCondition = "bEnableEditorPeformanceTool"))
