Fix for localization issue

diff --git a/Engine/Plugins/Experimental/EditorPerformance/Source/EditorPerformance/Private/SEditorPerformanceDialogs.cpp b/Engine/Plugins/Experimental/EditorPerformance/Source/EditorPerformance/Private/SEditorPerformanceDialogs.cpp
--- a/Engine/Plugins/Experimental/EditorPerformance/Source/EditorPerformance/Private/SEditorPerformanceDialogs.cpp
+++ b/Engine/Plugins/Experimental/EditorPerformance/Source/EditorPerformance/Private/SEditorPerformanceDialogs.cpp
@@ -531,7 +531,7 @@ TSharedRef<SWidget> SEditorPerformanceReportDialog::GetKPIGridPanel()
 				.Margin(DefaultMargin)
 				.ColorAndOpacity(EStyleColor::Foreground)
 				.Font(TitleFont)
-				.Text(LOCTEXT("ExpectedValueColumn", "Failures"))
+				.Text(LOCTEXT("FailedValueColumn", "Failures"))
 		];
 
 	if (EnableNotifcations)
