Set the correct menu insert for Editor Performance Tool

diff --git a/Engine/Plugins/Experimental/EditorPerformance/Source/EditorPerformance/Private/EditorPerformanceModule.cpp b/Engine/Plugins/Experimental/EditorPerformance/Source/EditorPerformance/Private/EditorPerformanceModule.cpp
--- a/Engine/Plugins/Experimental/EditorPerformance/Source/EditorPerformance/Private/EditorPerformanceModule.cpp
+++ b/Engine/Plugins/Experimental/EditorPerformance/Source/EditorPerformance/Private/EditorPerformanceModule.cpp
@@ -156,7 +156,7 @@ void FEditorPerformanceModule::InitializeEditor()
 
 	UToolMenu* Menu = UToolMenus::Get()->ExtendMenu("LevelEditor.StatusBar.ToolBar");
 
-	FToolMenuSection& EditorPerfSection = Menu->AddSection("Perf", FText::GetEmpty(), FToolMenuInsert("DDC", EToolMenuInsertType::Before));
+	FToolMenuSection& EditorPerfSection = Menu->AddSection("Perf", FText::GetEmpty(), FToolMenuInsert("Perf", EToolMenuInsertType::Before));
 
 	EditorPerfSection.AddEntry(FToolMenuEntry::InitWidget("EditorPerformanceStatusBar", CreateStatusBarWidget(), FText::GetEmpty(), true, false));
 }
