[StallLogger] Changing the icon for Stall Log, giving it notification badge overlay functionality, and general polish

diff --git a/Engine/Plugins/Experimental/EditorPerformance/Source/StallLogSubsystem/Private/StallLogSubsystem.cpp b/Engine/Plugins/Experimental/EditorPerformance/Source/StallLogSubsystem/Private/StallLogSubsystem.cpp
--- a/Engine/Plugins/Experimental/EditorPerformance/Source/StallLogSubsystem/Private/StallLogSubsystem.cpp
+++ b/Engine/Plugins/Experimental/EditorPerformance/Source/StallLogSubsystem/Private/StallLogSubsystem.cpp
@@ -134,6 +134,10 @@ public:
 	void OnStallDetected(uint64 UniqueID, FDateTime InDetectTime, TConstArrayView<uint64> Backtrace);
 	void OnStallCompleted(uint32 ThreadID, FStringView StatName, uint64 UniqueID, FDateTime InCompletedTime, float InDurationSeconds, bool bWasDetectedWithCallstack);
 
+	void SetHasViewedStalls(bool bHasViewed) { bHasViewedStalls = bHasViewed; };
+	bool GetHasViewedStalls() { return bHasViewedStalls; };
+
+	void AddStallToLog(FStallLogItemPtr LogItemToAdd);
 	void ClearStallLog();
 	
 	const TArray<FStallLogItemPtr>& GetStallLog() const;
@@ -143,6 +147,8 @@ private:
 	TMap<uint64, FStallLogItem> InFlightStalls;
 	// Completed stalls with full information
 	TArray<FStallLogItemPtr> StallLogs;
+	// If the user has viewed all current stall logs, set to false when a new stall is added
+	bool bHasViewedStalls = false;
 };
 
 const TArray<FStallLogItemPtr>& FStallLogHistory::GetStallLog() const
@@ -155,6 +161,7 @@ void FStallLogHistory::ClearStallLog()
 {
 	checkf(IsInGameThread(), TEXT("Can only be run on GameThread"));
 	StallLogs.Empty();
+	SetHasViewedStalls(true);
 }
 
 namespace
@@ -246,7 +253,7 @@ namespace
 					SNew(SButton)
 					.ToolTipText(LOCTEXT("StallDetector", "Copy Stall Information"))
 					.ButtonStyle(FAppStyle::Get(), "NoBorder")
-					//.ContentPadding(0.0)
+					.ContentPadding(FMargin(0.0f))
 					.Visibility(EVisibility::Visible)
 					.OnClicked_Lambda([this]()
 					{
@@ -321,6 +328,7 @@ namespace
 					[
 						SNew(SImage)
 						.Image(FAppStyle::GetBrush("GenericCommands.Copy"))
+						.DesiredSizeOverride(FVector2D(16.0f, 16.f))
 					]
 				]
 				+SHorizontalBox::Slot()
@@ -462,6 +470,7 @@ namespace
 	{
 		SLATE_BEGIN_ARGS(SStallLogStatusBarWidget) {}
 			SLATE_ARGUMENT(FGetSlateBrush, GetFilterBadgeIcon)
+			SLATE_ARGUMENT(FGetSlateBrush, GetFilterBadgeNotificationIcon)
 		SLATE_END_ARGS()
 
 		void Construct(const FArguments& InArgs);
@@ -470,8 +479,10 @@ namespace
 		
 		FText						GetToolTipText() const;
 		const FSlateBrush*			GetBadgeIcon() const;
+		const FSlateBrush*			GetBadgeNotificationIcon() const;
 
 		FGetSlateBrush GetFilterBadgeIconDelegate;
+		FGetSlateBrush GetFilterBadgeNotificationIconDelegate;
 	};
 
 	FText SStallLogStatusBarWidget::GetToolTipText() const
@@ -488,9 +499,19 @@ namespace
 		return nullptr;
 	}
 
+	const FSlateBrush* SStallLogStatusBarWidget::GetBadgeNotificationIcon() const
+	{
+		if (GetFilterBadgeNotificationIconDelegate.IsBound())
+		{
+			return GetFilterBadgeNotificationIconDelegate.Execute();
+		}
+		return nullptr;
+	}
+
 	void SStallLogStatusBarWidget::Construct(const FArguments& InArgs)
 	{
 		GetFilterBadgeIconDelegate = InArgs._GetFilterBadgeIcon;
+		GetFilterBadgeNotificationIconDelegate = InArgs._GetFilterBadgeNotificationIcon;
 		this->ChildSlot
 		[
 			SNew(SButton)
@@ -514,6 +535,13 @@ namespace
 						.Image_Lambda([this] { return GetBadgeIcon();  })
 						.ToolTipText_Lambda([this] { return GetToolTipText(); })
 					]
+					+ SOverlay::Slot()
+					.HAlign(HAlign_Center)
+					.VAlign(VAlign_Top)
+					[
+						SNew(SImage)
+							.Image_Lambda([this] { return GetBadgeNotificationIcon();  })
+					]
 				]
 			]
 			.OnClicked_Lambda([]()
@@ -537,7 +565,7 @@ void FStallLogHistory::OnStallDetected(
 	if (InFlightStalls.RemoveAndCopyValue(UniqueID, Existing))
 	{
 		// Already handled completed event, just add back trace & detect time
-		StallLogs.Emplace(MakeShared<FStallLogItem>(MoveTemp(Existing), InDetectTime, Backtrace));
+		AddStallToLog(MakeShared<FStallLogItem>(MoveTemp(Existing), InDetectTime, Backtrace));
 	}
 	else
 	{
@@ -558,7 +586,7 @@ void FStallLogHistory::OnStallCompleted(uint32 ThreadID, FStringView StatName, u
 		{
 			// Fill in final duration and add to list
 			// Existing->DurationSeconds = InDurationSeconds;
-			StallLogs.Emplace(MakeShared<FStallLogItem>(MoveTemp(Existing), StatName, MoveTemp(ThreadName), InDurationSeconds));
+			AddStallToLog(MakeShared<FStallLogItem>(MoveTemp(Existing), StatName, MoveTemp(ThreadName), InDurationSeconds));
 		}
 		else
 		{
@@ -570,7 +598,7 @@ void FStallLogHistory::OnStallCompleted(uint32 ThreadID, FStringView StatName, u
 	else
 	{
 		// No backtrace for this event
-		StallLogs.Emplace(MakeShared<FStallLogItem>(
+		AddStallToLog(MakeShared<FStallLogItem>(
 			StatName,
 			MoveTemp(ThreadName),
 			InDurationSeconds,
@@ -578,6 +606,16 @@ void FStallLogHistory::OnStallCompleted(uint32 ThreadID, FStringView StatName, u
 	}
 }
 
+void FStallLogHistory::AddStallToLog(FStallLogItemPtr LogItemToAdd)
+{
+	StallLogs.Emplace(LogItemToAdd);
+	// If the tab is open, we want to automatically say we viewed any new stalls
+	if (!FGlobalTabmanager::Get()->FindExistingLiveTab(FTabId(StallLogTabName)).IsValid())
+	{
+		SetHasViewedStalls(false);
+	}
+}
+
 DECLARE_CYCLE_STAT(
 	TEXT("StallLoggerSubsystem"),
 	STAT_FDelegateGraphTask_StallLogger,
@@ -627,7 +665,7 @@ void UStallLogSubsystem::Initialize(FSubsystemCollectionBase& Collection)
 
 	/** Register a tab spawner invoked by the tools bar button */
 	{
-		const FSlateIcon StallLogIcon(FAppStyle::GetAppStyleSetName(), "EditorViewport.ToggleRealTime");
+		const FSlateIcon StallLogIcon(FAppStyle::GetAppStyleSetName(), "EditorPerformance.Stall");
 
 		FGlobalTabmanager::Get()->RegisterNomadTabSpawner(StallLogTabName, FOnSpawnTab::CreateUObject(this, &UStallLogSubsystem::CreateStallLogTab))
 			.SetDisplayName(LOCTEXT("StallLogTabTitle", "Stall Log"))
@@ -645,7 +683,8 @@ void UStallLogSubsystem::Initialize(FSubsystemCollectionBase& Collection)
 	auto CreateStallLogWidget = [this]()
 	{
 		return SNew(SStallLogStatusBarWidget)
-			.GetFilterBadgeIcon(FGetSlateBrush::CreateUObject(this, &UStallLogSubsystem::GetStatusBarBadgeIcon));
+			.GetFilterBadgeIcon(FGetSlateBrush::CreateUObject(this, &UStallLogSubsystem::GetStatusBarBadgeIcon))
+			.GetFilterBadgeNotificationIcon(FGetSlateBrush::CreateUObject(this, &UStallLogSubsystem::GetStatusBarBadgeNotificationIcon));
 	};
 	
 	StallDetectorSection.AddEntry(
@@ -667,6 +706,7 @@ void UStallLogSubsystem::Deinitialize()
 	
 TSharedRef<SDockTab> UStallLogSubsystem::CreateStallLogTab(const FSpawnTabArgs& InArgs)
 {
+	StallLogHistory->SetHasViewedStalls(true);
 	return SAssignNew(StallLogTab, SDockTab)
 	.TabRole(ETabRole::NomadTab)
 	[
@@ -685,14 +725,17 @@ TSharedRef<SDockTab> UStallLogSubsystem::CreateStallLogTab(const FSpawnTabArgs&
 
 const FSlateBrush* UStallLogSubsystem::GetStatusBarBadgeIcon() const
 {
-	if (!StallLogHistory->GetStallLog().IsEmpty())
-	{
-		return FAppStyle::GetBrush(TEXT("EditorViewport.ToggleRealTime"));
-	}
-	else
+	return FAppStyle::GetBrush(TEXT("EditorPerformance.Stall"));
+}
+
+const FSlateBrush* UStallLogSubsystem::GetStatusBarBadgeNotificationIcon() const
+{
+	// If the StallLog has any unviewed stalls, display the notification overlay
+	if (!StallLogHistory->GetHasViewedStalls())
 	{
-		return FAppStyle::GetBrush(TEXT("Level.SaveDisabledIcon16x"));
+		return FAppStyle::GetBrush(TEXT("Icons.BadgeModified"));
 	}
+	return nullptr;
 }
 
 void UStallLogSubsystem::RegisterStallDetectedDelegates()
diff --git a/Engine/Plugins/Experimental/EditorPerformance/Source/StallLogSubsystem/Public/StallLogSubsystem.h b/Engine/Plugins/Experimental/EditorPerformance/Source/StallLogSubsystem/Public/StallLogSubsystem.h
--- a/Engine/Plugins/Experimental/EditorPerformance/Source/StallLogSubsystem/Public/StallLogSubsystem.h
+++ b/Engine/Plugins/Experimental/EditorPerformance/Source/StallLogSubsystem/Public/StallLogSubsystem.h
@@ -34,6 +34,7 @@ private:
 
 	TSharedRef<SDockTab> CreateStallLogTab(const FSpawnTabArgs& InArgs);
 	const FSlateBrush* GetStatusBarBadgeIcon() const;
+	const FSlateBrush* GetStatusBarBadgeNotificationIcon() const;
 
 	void RegisterStallDetectedDelegates();
 	void UnregisterStallDetectedDelegates();
diff --git a/Engine/Source/Editor/EditorStyle/Private/StarshipStyle.cpp b/Engine/Source/Editor/EditorStyle/Private/StarshipStyle.cpp
--- a/Engine/Source/Editor/EditorStyle/Private/StarshipStyle.cpp
+++ b/Engine/Source/Editor/EditorStyle/Private/StarshipStyle.cpp
@@ -1125,6 +1125,7 @@ void FStarshipEditorStyle::FStyle::SetupGeneralStyles()
 		Set("EditorPerformance.Notification.Good", new IMAGE_BRUSH_SVG("Starship/DerivedData/DD_RemoteCache_Idle", Icon16x16, EStyleColor::Success));
 		Set("EditorPerformance.Report.Panel", new IMAGE_BRUSH_SVG("Starship/DerivedData/DD_Cache_Statistics", Icon16x16));
 		Set("EditorPerformance.Report.Warning", new CORE_IMAGE_BRUSH_SVG("Starship/Common/alert-triangle", Icon16x16, EStyleColor::Warning));
+		Set("EditorPerformance.Stall", new IMAGE_BRUSH_SVG("Starship/Common/Stall", Icon16x16));
 	}
 
 	// Scene Outliner
