Fix warning when registering the toolbar item for the stall log subsystem
- It was trying to add itself to the DDC section. The DDC section may or may not exist. It gets replaced with a Zen section when Zen is in use. Also, the menu item may try to resolve its location before the level editor has registered its sections.
- The widget is now added directly to the bar instead of poking itself into sections created by another module.
- Also fixed initialization so that the widget doesn't start with a notification because it thinks there are stalls that haven't been viewed.

diff --git a/Engine/Plugins/Experimental/EditorPerformance/Source/StallLogSubsystem/Private/StallLogSubsystem.cpp b/Engine/Plugins/Experimental/EditorPerformance/Source/StallLogSubsystem/Private/StallLogSubsystem.cpp
--- a/Engine/Plugins/Experimental/EditorPerformance/Source/StallLogSubsystem/Private/StallLogSubsystem.cpp
+++ b/Engine/Plugins/Experimental/EditorPerformance/Source/StallLogSubsystem/Private/StallLogSubsystem.cpp
@@ -148,7 +148,7 @@ private:
 	// Completed stalls with full information
 	TArray<FStallLogItemPtr> StallLogs;
 	// If the user has viewed all current stall logs, set to false when a new stall is added
-	bool bHasViewedStalls = false;
+	bool bHasViewedStalls = true;
 };
 
 const TArray<FStallLogItemPtr>& FStallLogHistory::GetStallLog() const
@@ -677,8 +677,7 @@ void UStallLogSubsystem::Initialize(FSubsystemCollectionBase& Collection)
 	/** Add the widget button to the tools bar */	
 	UToolMenu* Menu = UToolMenus::Get()->ExtendMenu("LevelEditor.StatusBar.ToolBar");
 
-	FToolMenuSection& StallDetectorSection = Menu->AddSection(
-		"StallLog", FText::GetEmpty(), FToolMenuInsert("DDC", EToolMenuInsertType::Before));
+	FToolMenuSection& StallDetectorSection = Menu->AddSection("StallLog", FText::GetEmpty());
 
 	auto CreateStallLogWidget = [this]()
 	{
