Removed Total Availab le Memory KPI. All that matters is total free memory
Fixed issue with Editor Performance STatus bar not appearing when enabled

diff --git a/Engine/Plugins/Experimental/EditorPerformance/Source/EditorPerformance/Private/EditorPerformanceModule.cpp b/Engine/Plugins/Experimental/EditorPerformance/Source/EditorPerformance/Private/EditorPerformanceModule.cpp
--- a/Engine/Plugins/Experimental/EditorPerformance/Source/EditorPerformance/Private/EditorPerformanceModule.cpp
+++ b/Engine/Plugins/Experimental/EditorPerformance/Source/EditorPerformance/Private/EditorPerformanceModule.cpp
@@ -139,7 +139,7 @@ void FEditorPerformanceModule::InitializeEditor()
 	const FSlateIcon PerformanceReportIcon(FAppStyle::GetAppStyleSetName(), "EditorPerformance.Report.Panel");
 
 	FGlobalTabmanager::Get()->RegisterNomadTabSpawner(EditorPerformanceReportTabName, FOnSpawnTab::CreateRaw(this, &FEditorPerformanceModule::CreatePerformanceReportTab))
-		.SetDisplayName(LOCTEXT("EditorPerformanceReportTabTitle", "Editor Performance"))
+		.SetDisplayName(LOCTEXT("EditorPerformanceReportTabTitle", "Performance"))
 		.SetTooltipText(LOCTEXT("EditorPerformanceReportTabToolTipText", "Opens the Editor Performance Report tab."))
 		.SetGroup(WorkspaceMenu::GetMenuStructure().GetDeveloperToolsProfilingCategory())
 		.SetIcon(PerformanceReportIcon);
@@ -156,9 +156,13 @@ void FEditorPerformanceModule::InitializeEditor()
 
 	UToolMenu* Menu = UToolMenus::Get()->ExtendMenu("LevelEditor.StatusBar.ToolBar");
 
-	FToolMenuSection& EditorPerfSection = Menu->AddSection("Perf", FText::GetEmpty(), FToolMenuInsert("Perf", EToolMenuInsertType::Before));
-
-	EditorPerfSection.AddEntry(FToolMenuEntry::InitWidget("EditorPerformanceStatusBar", CreateStatusBarWidget(), FText::GetEmpty(), true, false));
+	if (Menu != nullptr)
+	{
+		// Add the Editor Perf ToolBar 
+		FToolMenuSection& EditorPerfSection = Menu->AddSection("EditorPerf", FText::GetEmpty(), FToolMenuInsert("Compile", EToolMenuInsertType::Before));
+		EditorPerfSection.AddEntry(FToolMenuEntry::InitWidget("EditorPerformanceStatusBar", CreateStatusBarWidget(), FText::GetEmpty(), true, false));
+	}
+	
 }
 
 void FEditorPerformanceModule::TerminateEditor()
@@ -219,21 +223,21 @@ void FEditorPerformanceModule::InitializeKPIs()
 	LocalDDCEfficiencyKPI		= KPIRegistry.DeclareKPIValue(CacheCategoryName, LocalDDCEfficiencyKPIName, 100.0, LocalDDCEffciencyKPILimit, FKPIValue::GreaterThan, FKPIValue::Percent);
 	VirtualAssetEfficiencyKPI	= KPIRegistry.DeclareKPIValue(CacheCategoryName, VirtualAssetEfficiencyKPIName, 100.0, VirtualAssetEfficiencyKPILimit, FKPIValue::GreaterThan, FKPIValue::Percent);
 	CoreCountKPI				= KPIRegistry.DeclareKPIValue(HardwareCategoryName, CoreCountKPIName, 128.0, CoreCountKPILimit, FKPIValue::GreaterThanOrEqual, FKPIValue::Decimal);
-	TotalMemoryKPI				= KPIRegistry.DeclareKPIValue(HardwareCategoryName, TotalMemoryKPIName, 128.0, TotalMemoryKPILimit, FKPIValue::GreaterThanOrEqual, FKPIValue::GigaBytes);
+	//TotalMemoryKPI				= KPIRegistry.DeclareKPIValue(HardwareCategoryName, TotalMemoryKPIName, 128.0, TotalMemoryKPILimit, FKPIValue::GreaterThanOrEqual, FKPIValue::GigaBytes);
 	AvailableMemoryKPI			= KPIRegistry.DeclareKPIValue(HardwareCategoryName, AvailableMemoryKPIName, 128.0, AvailableMemoryKPILimit, FKPIValue::GreaterThanOrEqual, FKPIValue::GigaBytes);
 	
 	// Declare the KPI Hints
-	KPIRegistry.DeclareKPIHint(EditorBootKPI, LOCTEXT("EditorBootHintMessage", "The Editor boot time is slow.\nCheck you have enbabled a Game Feature Plugin profile for your project and that the expected local cache efficiency is met.\nIf you are booting the Editor in the background then disable the Use Less CPU in Background option in the settings."), LOCTEXT("EditorBootHintURL","https://docs.unrealengine.com/5.0/en-US/"));
-	KPIRegistry.DeclareKPIHint(TotalTimeToEditorKPI, LOCTEXT("EditorStartupHintMessage", "The Editor start-up time is slow.\nCheck you have enbabled a Game Feature Plugin profile for your project and that the expected local cache efficiency is met.\nIf you are booting the Editor in the background then disable the Use Less CPU in Background option in the settings."), LOCTEXT("EditorBootHintURL","https://docs.unrealengine.com/5.0/en-US/"));
-	KPIRegistry.DeclareKPIHint(EditorPluginCountKPI, LOCTEXT("EditorPluginHintMessage", "The Editor is loading more plugins than expected and this will affect Editor start-up performance.\nCheck you have enbabled a Game Feature Plugin profile for your project."), LOCTEXT("EditorPluginHintURL", "https://docs.unrealengine.com/5.0/en-US/"));
+	KPIRegistry.DeclareKPIHint(EditorBootKPI, LOCTEXT("EditorBootHintMessage", "The Editor boot time is slow.\nCheck you have enabled a Game Feature Plugin profile for your project and that the expected local cache efficiency is met.\nIf you are booting the Editor in the background then disable the Use Less CPU in Background option in the settings."), LOCTEXT("EditorBootHintURL","https://docs.unrealengine.com/5.0/en-US/"));
+	KPIRegistry.DeclareKPIHint(TotalTimeToEditorKPI, LOCTEXT("EditorStartupHintMessage", "The Editor start-up time is slow.\nCheck you have enabled a Game Feature Plugin profile for your project and that the expected local cache efficiency is met.\nIf you are booting the Editor in the background then disable the Use Less CPU in Background option in the settings."), LOCTEXT("EditorBootHintURL","https://docs.unrealengine.com/5.0/en-US/"));
+	KPIRegistry.DeclareKPIHint(EditorPluginCountKPI, LOCTEXT("EditorPluginHintMessage", "The Editor is loading more plugins than expected and this will affect Editor start-up performance.\nCheck you have enabled a Game Feature Plugin profile for your project."), LOCTEXT("EditorPluginHintURL", "https://docs.unrealengine.com/5.0/en-US/"));
 
-	KPIRegistry.DeclareKPIHint(PIETransitionKPI, LOCTEXT("PIETransitionHintMessage", "The Editor transtion to PIE is slow.\nCheck that the expected local cache efficiency is met.\nIf you are transitioning to PIE with the Editor in the background then disable the Use Less CPU in Background option in the settings."), LOCTEXT("PIETransitionHintURL", "https://docs.unrealengine.com/5.0/en-US/"));
+	KPIRegistry.DeclareKPIHint(PIETransitionKPI, LOCTEXT("PIETransitionHintMessage", "The Editor transition to PIE is slow.\nCheck that the expected local cache efficiency is met.\nIf you are transitioning to PIE with the Editor in the background then disable the Use Less CPU in Background option in the settings."), LOCTEXT("PIETransitionHintURL", "https://docs.unrealengine.com/5.0/en-US/"));
 	
 	KPIRegistry.DeclareKPIHint(LocalDDCEfficiencyKPI, LOCTEXT("LocalCacheEfficencyHintMessage", "The Editor will not perform well if the local cache efficiency has not yet met the expected value.\nIf this is the first time you have booted the Editor after a sync then this is to be expected."), LOCTEXT("EditorCacheHintURL","https://docs.unrealengine.com/5.3/en-US/derived-data-cache/"));
 	
-	KPIRegistry.DeclareKPIHint(CoreCountKPI, LOCTEXT("LowCoreCountHintMessage", "Your hardware has a low CPU core count.\nUsing a lower than recommended hardware specifcation for development is not recommended for good developer efficiency."), LOCTEXT("LowCoreCountHintURL", "https://docs.unrealengine.com/5.0/en-US/"));
-	KPIRegistry.DeclareKPIHint(TotalMemoryKPI, LOCTEXT("LowTotalMemoryHintMessage", "Your hardware has a low Total Memory.\nUsing a lower than recommended hardware specifcation for development is not recommended for good developer efficiency."), LOCTEXT("LowTotalMemoryHintURL", "https://docs.unrealengine.com/5.0/en-US/"));
-	KPIRegistry.DeclareKPIHint(AvailableMemoryKPI, LOCTEXT("LowAvaliableMemoryHintMessage", "You hardware is running low on memory.\nTry closing applciations that are no longer needed to recover available memory."), LOCTEXT("LowAvailableMemoryHintURL", "https://docs.unrealengine.com/5.0/en-US/"));
+	KPIRegistry.DeclareKPIHint(CoreCountKPI, LOCTEXT("LowCoreCountHintMessage", "Your hardware has a low CPU core count.\nUsing a lower than recommended hardware specification for development is not recommended for good developer efficiency."), LOCTEXT("LowCoreCountHintURL", "https://docs.unrealengine.com/5.0/en-US/"));
+	KPIRegistry.DeclareKPIHint(TotalMemoryKPI, LOCTEXT("LowTotalMemoryHintMessage", "Your hardware has a low Total Memory.\nUsing a lower than recommended hardware specification for development is not recommended for good developer efficiency."), LOCTEXT("LowTotalMemoryHintURL", "https://docs.unrealengine.com/5.0/en-US/"));
+	KPIRegistry.DeclareKPIHint(AvailableMemoryKPI, LOCTEXT("LowAvaliableMemoryHintMessage", "You hardware is running low on memory.\nTry closing applications that are no longer needed to recover available memory."), LOCTEXT("LowAvailableMemoryHintURL", "https://docs.unrealengine.com/5.0/en-US/"));
 
 	// Load the KPI profiles
 	KPIRegistry.LoadKPIProfiles(TEXT("EditorPerformance.Profile"), GEditorIni);
@@ -252,7 +256,7 @@ void FEditorPerformanceModule::InitializeKPIs()
 
 	// Gather hardware stats
 	KPIRegistry.SetKPIValue(CoreCountKPI, float(FPlatformMisc::NumberOfCores()));
-	KPIRegistry.SetKPIValue(TotalMemoryKPI, static_cast<float>(FPlatformMemory::GetStats().TotalPhysical) / (1024.0f * 1024.0f * 1024.0f));
+	KPIRegistry.SetKPIValue(TotalMemoryKPI, (float)FMath::CeilToInt(static_cast<float>(FPlatformMemory::GetStats().TotalPhysical) / (1024.0f * 1024.0f * 1024.0f)));
 
 	EditorState = EEditorState::Editor_Boot;
 
