Editor Performance KPI: Added cvar to force a KPI into a bad state for debug purposes.

diff --git a/Engine/Plugins/Experimental/EditorPerformance/Source/EditorPerformance/Private/EditorPerformanceModule.cpp b/Engine/Plugins/Experimental/EditorPerformance/Source/EditorPerformance/Private/EditorPerformanceModule.cpp
--- a/Engine/Plugins/Experimental/EditorPerformance/Source/EditorPerformance/Private/EditorPerformanceModule.cpp
+++ b/Engine/Plugins/Experimental/EditorPerformance/Source/EditorPerformance/Private/EditorPerformanceModule.cpp
@@ -109,6 +109,13 @@ float MemoryPressureKPILimit = (float)FGenericPlatformMemoryStats::EMemoryPressu
  
 static const FName EditorPerformanceReportTabName = FName(TEXT("EditorPerformanceReportTab"));
 
+static FString GForceBadKPIState = TEXT("");
+static FAutoConsoleVariableRef CVarForceBadKPIState(
+	TEXT("Editor.PerformanceKPI.ForceBadState"),
+	GForceBadKPIState,
+	TEXT("The name (category_name) of a KPI to force into a bad state. The KPI will stay in a bad state as long as this cvar isn't cleared."),
+	ECVF_Default);
+
 void FEditorPerformanceModule::StartupModule()
 {
 	InitializeKPIs();
@@ -650,6 +657,25 @@ void FEditorPerformanceModule::UpdateKPIs(float InDeltaTime)
 			RecordedKPIEvent.Remove(KPIValue.Id);
 		}
 	}
+
+	DebugForceKPIIntoBadStateIfNeeded();
+}
+
+void FEditorPerformanceModule::DebugForceKPIIntoBadStateIfNeeded()
+{
+	if (!GForceBadKPIState.IsEmpty())
+	{
+		for (FKPIValues::TConstIterator KPIValueIt(GetKPIRegistry().GetKPIValues()); KPIValueIt; ++KPIValueIt)
+		{
+			const FKPIValue& KPIValue = KPIValueIt->Value;
+
+			if (KPIValue.Path == FName(GForceBadKPIState))
+			{
+				const_cast<FKPIValue&>(KPIValue).State = FKPIValue::Bad;
+				break;
+			}
+		}
+	}
 }
 
 FText FEditorPerformanceModule::GetMemoryPressureDisplayValue(const float InValue) const
diff --git a/Engine/Plugins/Experimental/EditorPerformance/Source/EditorPerformance/Public/EditorPerformanceModule.h b/Engine/Plugins/Experimental/EditorPerformance/Source/EditorPerformance/Public/EditorPerformanceModule.h
--- a/Engine/Plugins/Experimental/EditorPerformance/Source/EditorPerformance/Public/EditorPerformanceModule.h
+++ b/Engine/Plugins/Experimental/EditorPerformance/Source/EditorPerformance/Public/EditorPerformanceModule.h
@@ -73,6 +73,8 @@ private:
 	UE_API void HeartBeatCallback();
 	UE_API void HitchSamplerCallback();
 
+	void DebugForceKPIIntoBadStateIfNeeded();
+
 	FText GetMemoryPressureDisplayValue(const float) const;
 
 	FKPIRegistry					KPIRegistry;
