Merge stall logs and performance KPI into a single panel summoned from the status bar
- Removed the stall log button from the status bar

diff --git a/Engine/Plugins/Experimental/EditorPerformance/Source/EditorPerformance/EditorPerformance.Build.cs b/Engine/Plugins/Experimental/EditorPerformance/Source/EditorPerformance/EditorPerformance.Build.cs
--- a/Engine/Plugins/Experimental/EditorPerformance/Source/EditorPerformance/EditorPerformance.Build.cs
+++ b/Engine/Plugins/Experimental/EditorPerformance/Source/EditorPerformance/EditorPerformance.Build.cs
@@ -28,7 +28,8 @@ public class EditorPerformance : ModuleRules
 				"ToolWidgets",
 				"DerivedDataCache",
 				"Virtualization",
-				"StudioTelemetry"
+				"StudioTelemetry",
+				"StallLogSubsystem"
 			});
 
 		PrivateIncludePathModuleNames.AddRange
diff --git a/Engine/Plugins/Experimental/EditorPerformance/Source/EditorPerformance/Private/EditorPerformanceModule.cpp b/Engine/Plugins/Experimental/EditorPerformance/Source/EditorPerformance/Private/EditorPerformanceModule.cpp
--- a/Engine/Plugins/Experimental/EditorPerformance/Source/EditorPerformance/Private/EditorPerformanceModule.cpp
+++ b/Engine/Plugins/Experimental/EditorPerformance/Source/EditorPerformance/Private/EditorPerformanceModule.cpp
@@ -4,6 +4,7 @@
 
 #include "Framework/Application/SlateApplication.h"
 #include "Modules/ModuleManager.h"
+#include "SEditorDiagnosticsTab.h"
 #include "SEditorPerformanceDialogs.h"
 #include "SEditorPerformanceStatusBar.h"
 #include "Styling/AppStyle.h"
@@ -107,7 +108,7 @@ float LocalDDCEffciencyKPILimit = 85.f;
 float VirtualAssetEfficiencyKPILimit = 95.f;
 float MemoryPressureKPILimit = (float)FGenericPlatformMemoryStats::EMemoryPressureStatus::Warning;
  
-static const FName EditorPerformanceReportTabName = FName(TEXT("EditorPerformanceReportTab"));
+static const FName EditorDiagnosticsTabName = FName(TEXT("EditorDiagnosticsTab"));
 
 static FString GForceBadKPIState = TEXT("");
 static FAutoConsoleVariableRef CVarForceBadKPIState(
@@ -131,24 +132,6 @@ void FEditorPerformanceModule::InitializeEditor()
 {
 	UEditorPerformanceSettings* EditorPerformanceSettings = GetMutableDefault<UEditorPerformanceSettings>();
 
-	// Check if we want to have the tool enabled or not.
-	if (EditorPerformanceSettings && EditorPerformanceSettings->bEnableEditorPeformanceTool==false )
-	{
-		return;
-	}
-
-	if (EditorPerformanceSettings)
-	{
-		// Enable any experimental features
-		if (EditorPerformanceSettings->bEnableExperimentalFeatures)
-		{
-			CloudDDCLatencyKPI			= KPIRegistry.DeclareKPIValue(CacheCategoryName, CacheCategoryText, CloudDDCLatencyKPIName, CloudDDCLatencyKPIText, 0.0, NoSetLimit, FKPIValue::LessThan, FKPIValue::Milliseconds);
-			CloudDDCReadSpeedKPI		= KPIRegistry.DeclareKPIValue(CacheCategoryName, CacheCategoryText, CloudDDCReadSpeedKPIName, CloudDDCReadSpeedKPIText, 100.0, NoSetLimit, FKPIValue::GreaterThan, FKPIValue::MegaBitsPerSecond);
-			EditorStallRateKPI			= KPIRegistry.DeclareKPIValue(EditorCategoryName, EditorCategoryText, EditorStallRateKPIName, EditorStallRateKPIText, 0.0, NoSetLimit, FKPIValue::LessThan, FKPIValue::Percent);
-			PIEStallRateKPI				= KPIRegistry.DeclareKPIValue(PIECategoryName, PIECategoryText, PIEStallRateKPIName, PIEStallRateKPIText, 0.0, NoSetLimit, FKPIValue::LessThan, FKPIValue::Percent);
-		}
-	}
-
 	FTimerDelegate HeartBeatDelegate;
 	HeartBeatDelegate.BindRaw(this, &FEditorPerformanceModule::HeartBeatCallback);
 	GEditor->GetTimerManager()->SetTimer(HeartBeatTimerHandle, HeartBeatDelegate, HeartBeatIntervalSeconds, true);
@@ -159,9 +142,9 @@ void FEditorPerformanceModule::InitializeEditor()
 
 	const FSlateIcon PerformanceReportIcon(FAppStyle::GetAppStyleSetName(), "EditorPerformance.Report.Panel");
 
-	FGlobalTabmanager::Get()->RegisterNomadTabSpawner(EditorPerformanceReportTabName, FOnSpawnTab::CreateRaw(this, &FEditorPerformanceModule::CreatePerformanceReportTab))
-		.SetDisplayName(LOCTEXT("EditorPerformanceReportTabTitle", "Performance"))
-		.SetTooltipText(LOCTEXT("EditorPerformanceReportTabToolTipText", "Opens the Editor Performance Report tab."))
+	FGlobalTabmanager::Get()->RegisterNomadTabSpawner(EditorDiagnosticsTabName, FOnSpawnTab::CreateRaw(this, &FEditorPerformanceModule::CreatePerformanceReportTab))
+		.SetDisplayName(LOCTEXT("EditorDiagnosticsTabTitle", "Editor Diagnostics"))
+		.SetTooltipText(LOCTEXT("EditorDiagnosticsTabToolTipText", "Opens the Editor Diagnostics tab."))
 		.SetGroup(WorkspaceMenu::GetMenuStructure().GetDeveloperToolsProfilingCategory())
 		.SetIcon(PerformanceReportIcon);
 
@@ -173,7 +156,7 @@ void FEditorPerformanceModule::InitializeEditor()
 	}
 #endif // WITH_RELOAD
 
-	FEditorPerformanceStatusBarMenuCommands::Register();
+	FEditorDiagnosticsCommands::Register();
 
 	UToolMenu* Menu = UToolMenus::Get()->ExtendMenu("LevelEditor.StatusBar.ToolBar");
 
@@ -184,13 +167,30 @@ void FEditorPerformanceModule::InitializeEditor()
 		EditorPerfSection.AddEntry(FToolMenuEntry::InitWidget("EditorPerformanceStatusBar", CreateStatusBarWidget(), FText::GetEmpty(), true, false));
 	}
 	
+	// Check if we want to have the tool enabled or not.
+	if (EditorPerformanceSettings && EditorPerformanceSettings->bEnableEditorPeformanceTool == false )
+	{
+		return;
+	}
+
+	if (EditorPerformanceSettings)
+	{
+		// Enable any experimental features
+		if (EditorPerformanceSettings->bEnableExperimentalFeatures)
+		{
+			CloudDDCLatencyKPI			= KPIRegistry.DeclareKPIValue(CacheCategoryName, CacheCategoryText, CloudDDCLatencyKPIName, CloudDDCLatencyKPIText, 0.0, NoSetLimit, FKPIValue::LessThan, FKPIValue::Milliseconds);
+			CloudDDCReadSpeedKPI		= KPIRegistry.DeclareKPIValue(CacheCategoryName, CacheCategoryText, CloudDDCReadSpeedKPIName, CloudDDCReadSpeedKPIText, 100.0, NoSetLimit, FKPIValue::GreaterThan, FKPIValue::MegaBitsPerSecond);
+			EditorStallRateKPI			= KPIRegistry.DeclareKPIValue(EditorCategoryName, EditorCategoryText, EditorStallRateKPIName, EditorStallRateKPIText, 0.0, NoSetLimit, FKPIValue::LessThan, FKPIValue::Percent);
+			PIEStallRateKPI				= KPIRegistry.DeclareKPIValue(PIECategoryName, PIECategoryText, PIEStallRateKPIName, PIEStallRateKPIText, 0.0, NoSetLimit, FKPIValue::LessThan, FKPIValue::Percent);
+		}
+	}
 }
 
 void FEditorPerformanceModule::TerminateEditor()
 {
 	if (FSlateApplication::IsInitialized())
 	{
-		FGlobalTabmanager::Get()->UnregisterNomadTabSpawner(EditorPerformanceReportTabName);
+		FGlobalTabmanager::Get()->UnregisterNomadTabSpawner(EditorDiagnosticsTabName);
 
 		if (PerformanceReportTab.IsValid())
 		{
@@ -198,7 +198,7 @@ void FEditorPerformanceModule::TerminateEditor()
 		}
 	}
 
-	FEditorPerformanceStatusBarMenuCommands::Unregister();
+	FEditorDiagnosticsCommands::Unregister();
 }
 
 TSharedRef<SWidget> FEditorPerformanceModule::CreateStatusBarWidget()
@@ -206,24 +206,14 @@ TSharedRef<SWidget> FEditorPerformanceModule::CreateStatusBarWidget()
 	return SNew(SEditorPerformanceStatusBarWidget);
 }
 
-TSharedPtr<SWidget> FEditorPerformanceModule::CreatePerformanceReportDialog()
-{
-	return SNew(SEditorPerformanceReportDialog);
-}
-
-
 TSharedRef<SDockTab> FEditorPerformanceModule::CreatePerformanceReportTab(const FSpawnTabArgs& Args)
 {
-	return SAssignNew(PerformanceReportTab, SDockTab)
-		.TabRole(ETabRole::NomadTab)
-		[
-			CreatePerformanceReportDialog().ToSharedRef()
-		];
+	return SNew(SEditorDiagnosticsTab).TabRole(ETabRole::NomadTab);
 }
 
 void FEditorPerformanceModule::ShowPerformanceReportTab()
 {
-	FGlobalTabmanager::Get()->TryInvokeTab(FTabId(EditorPerformanceReportTabName));
+	FGlobalTabmanager::Get()->TryInvokeTab(FTabId(EditorDiagnosticsTabName));
 }
 
 void FEditorPerformanceModule::InitializeKPIs()
diff --git a/Engine/Plugins/Experimental/EditorPerformance/Source/EditorPerformance/Private/SEditorDiagnosticsTab.cpp b/Engine/Plugins/Experimental/EditorPerformance/Source/EditorPerformance/Private/SEditorDiagnosticsTab.cpp
new file mode 100644
--- /dev/null
+++ b/Engine/Plugins/Experimental/EditorPerformance/Source/EditorPerformance/Private/SEditorDiagnosticsTab.cpp
@@ -0,0 +1,161 @@
+// Copyright Epic Games, Inc. All Rights Reserved.
+
+#include "SEditorDiagnosticsTab.h"
+
+#include "SEditorPerformanceDialogs.h"
+#include "StallLogSubsystem.h"
+
+#include "Editor/EditorEngine.h"
+#include "Editor/EditorPerformanceSettings.h"
+#include "Framework/MultiBox/MultiBoxBuilder.h"
+#include "Styling/StyleColors.h"
+#include "Widgets/Layout/SBorder.h"
+#include "Widgets/Layout/SBox.h"
+#include "Widgets/Layout/SSplitter.h"
+
+#define LOCTEXT_NAMESPACE "EditorDiagnostics"
+
+extern UNREALED_API UEditorEngine* GEditor;
+
+FEditorDiagnosticsCommands::FEditorDiagnosticsCommands()
+	: TCommands<FEditorDiagnosticsCommands>(
+		TEXT("EditorDiagnostics"),
+		NSLOCTEXT("Contexts", "EditorDiagnostics", "Editor Diagnostics"),
+		NAME_None, // Parent
+		FAppStyle::GetAppStyleSetName() // Icon Style Set
+	)
+{
+}
+
+void FEditorDiagnosticsCommands::RegisterCommands()
+{	
+	UI_COMMAND(SetShowActivityMonitor, "Activity", "Displays the activity monitor", EUserInterfaceActionType::ToggleButton, FInputChord());
+	UI_COMMAND(SetShowStalls, "Stalls", "Displays the stalls log", EUserInterfaceActionType::ToggleButton, FInputChord());
+}
+
+void SEditorDiagnosticsTab::Construct(const FArguments& InArgs)
+{
+	SDockTab::Construct(InArgs);
+
+	TSharedRef<SBorder> ContentBorderRef =
+		SNew(SBorder)
+		.BorderImage(FAppStyle::Get().GetBrush("WhiteBrush"))
+		.BorderBackgroundColor(FSlateColor(EStyleColor::Panel));
+
+	ContentBorder = ContentBorderRef;
+
+	TSharedRef<SWidget> Splitter
+		= SNew(SBox)
+		.Padding(2.f)
+		[
+			SNew(SSplitter)
+			.Orientation(Orient_Horizontal)
+			.PhysicalSplitterHandleSize(2.0f)
+			+SSplitter::Slot()
+			.SizeRule(SSplitter::SizeToContent)
+			[
+				SNew(SBorder)
+				.Padding(0.f)
+				.BorderImage(FAppStyle::GetBrush("ToolPanel.GroupBorder"))
+				[
+					SNew(SBox)
+					.MinDesiredWidth(64.f)
+					[
+						CreateToolBar()
+					]
+				]
+			]
+			+SSplitter::Slot()
+			.SizeRule(SSplitter::FractionOfParent)
+			[
+				ContentBorderRef
+			]
+		];
+
+	SetContent(Splitter);
+}
+
+TSharedRef<SWidget> SEditorDiagnosticsTab::CreateToolBar()
+{
+	TSharedRef<FUICommandList> ToolBarCommandList = MakeShared<FUICommandList>();
+
+	bool bIsActivityMonitorAvailable = false;
+
+	UEditorPerformanceSettings* EditorPerformanceSettings = GetMutableDefault<UEditorPerformanceSettings>();
+	if (EditorPerformanceSettings && EditorPerformanceSettings->bEnableEditorPeformanceTool)
+	{
+		bIsActivityMonitorAvailable = true;
+
+		TWeakPtr<FUICommandInfo> ShowActivityMonitorCommand = FEditorDiagnosticsCommands::Get().SetShowActivityMonitor;
+
+		ToolBarCommandList->MapAction(FEditorDiagnosticsCommands::Get().SetShowActivityMonitor,
+			FExecuteAction::CreateRaw(this, &SEditorDiagnosticsTab::SetShowActivityMonitor),
+			FCanExecuteAction(),
+			FGetActionCheckState::CreateLambda([this, ShowActivityMonitorCommand] ()
+				{
+					return SelectedCommand == ShowActivityMonitorCommand ? ECheckBoxState::Checked : ECheckBoxState::Unchecked;		
+				}));
+	}
+
+	TWeakPtr<FUICommandInfo> ShowStallsCommand = FEditorDiagnosticsCommands::Get().SetShowStalls;
+
+	ToolBarCommandList->MapAction(FEditorDiagnosticsCommands::Get().SetShowStalls,
+		FExecuteAction::CreateRaw(this, &SEditorDiagnosticsTab::SetShowStalls),
+		FCanExecuteAction(),
+		FGetActionCheckState::CreateLambda([this, ShowStallsCommand] ()
+			{
+				return SelectedCommand == ShowStallsCommand ? ECheckBoxState::Checked : ECheckBoxState::Unchecked;		
+			}));
+
+	constexpr bool bForceSmallIcons = true; 
+	FVerticalToolBarBuilder ToolBarBuilder = FVerticalToolBarBuilder(ToolBarCommandList, FMultiBoxCustomization::None, TSharedPtr<FExtender>(), bForceSmallIcons);
+	ToolBarBuilder.SetLabelVisibility(EVisibility::Visible);
+
+	ToolBarBuilder.SetStyle(&FAppStyle::Get(), "CategoryDrivenContentBuilderToolbarWithLabels");
+
+	if (bIsActivityMonitorAvailable)
+	{
+		ToolBarBuilder.AddToolBarButton(FEditorDiagnosticsCommands::Get().SetShowActivityMonitor);
+		SetShowActivityMonitor();
+	}
+	else
+	{
+		SetShowStalls();
+	}
+
+	ToolBarBuilder.AddToolBarButton(FEditorDiagnosticsCommands::Get().SetShowStalls);
+
+	return ToolBarBuilder.MakeWidget();
+}
+
+TSharedRef<SWidget> SEditorDiagnosticsTab::CreateActivityMonitorPanel()
+{
+	return SNew(SEditorPerformanceReportDialog);
+}
+
+TSharedRef<SWidget> SEditorDiagnosticsTab::CreateStallLogPanel()
+{
+	UStallLogSubsystem* StallLogSubsystem = GEditor->GetEditorSubsystem<UStallLogSubsystem>();
+
+	return StallLogSubsystem->CreateStallLogPanel();
+}
+
+void SEditorDiagnosticsTab::SetShowActivityMonitor()
+{
+	if (TSharedPtr<SBorder> BorderWidget = ContentBorder.Pin())
+	{
+		BorderWidget->SetContent(CreateActivityMonitorPanel());
+		SelectedCommand = FEditorDiagnosticsCommands::Get().SetShowActivityMonitor;
+	}
+}
+
+void SEditorDiagnosticsTab::SetShowStalls()
+{
+	if (TSharedPtr<SBorder> BorderWidget = ContentBorder.Pin())
+	{
+		BorderWidget->SetContent(CreateStallLogPanel());
+		SelectedCommand = FEditorDiagnosticsCommands::Get().SetShowStalls;
+	}
+}
+
+#undef LOCTEXT_NAMESPACE
diff --git a/Engine/Plugins/Experimental/EditorPerformance/Source/EditorPerformance/Private/SEditorDiagnosticsTab.h b/Engine/Plugins/Experimental/EditorPerformance/Source/EditorPerformance/Private/SEditorDiagnosticsTab.h
new file mode 100644
--- /dev/null
+++ b/Engine/Plugins/Experimental/EditorPerformance/Source/EditorPerformance/Private/SEditorDiagnosticsTab.h
@@ -0,0 +1,43 @@
+// Copyright Epic Games, Inc. All Rights Reserved.
+
+#pragma once
+
+#include "Framework/Commands/Commands.h"
+#include "Widgets/Docking/SDockTab.h"
+
+class FUICommandInfo;
+class SBorder;
+class SWidget;
+
+class FEditorDiagnosticsCommands : public TCommands<FEditorDiagnosticsCommands>
+{
+public:
+	FEditorDiagnosticsCommands();
+
+	// TCommand<> interface
+	virtual void RegisterCommands() override;
+	// End of TCommand<> interface
+
+public:
+	TSharedPtr<FUICommandInfo> SetShowActivityMonitor;
+	TSharedPtr<FUICommandInfo> SetShowStalls;
+};
+
+class SEditorDiagnosticsTab : public SDockTab
+{
+public:
+	void Construct(const FArguments& InArgs);
+
+protected:
+	TSharedRef<SWidget> CreateToolBar();
+
+	TSharedRef<SWidget> CreateActivityMonitorPanel();
+	TSharedRef<SWidget> CreateStallLogPanel();
+
+	void SetShowActivityMonitor();
+	void SetShowStalls();
+
+private:
+	TWeakPtr<SBorder> ContentBorder;
+	TWeakPtr<FUICommandInfo> SelectedCommand;
+};
diff --git a/Engine/Plugins/Experimental/EditorPerformance/Source/EditorPerformance/Private/SEditorPerformanceStatusBar.cpp b/Engine/Plugins/Experimental/EditorPerformance/Source/EditorPerformance/Private/SEditorPerformanceStatusBar.cpp
--- a/Engine/Plugins/Experimental/EditorPerformance/Source/EditorPerformance/Private/SEditorPerformanceStatusBar.cpp
+++ b/Engine/Plugins/Experimental/EditorPerformance/Source/EditorPerformance/Private/SEditorPerformanceStatusBar.cpp
@@ -38,7 +38,7 @@
 #include "UObject/NameTypes.h"
 #include "UObject/UnrealNames.h"
 #include "Widgets/Images/SImage.h"
-#include "Widgets/Input/SComboButton.h"
+#include "Widgets/Input/SButton.h"
 #include "Widgets/Notifications/SNotificationList.h"
 #include "Widgets/SBoxPanel.h"
 #include "Widgets/SOverlay.h"
@@ -52,118 +52,43 @@ struct FSlateBrush;
 #define LOCTEXT_NAMESPACE "EditorPerformance"
 
 
-TSharedRef<FUICommandList> FEditorPerformanceStatusBarMenuCommands::ActionList(new FUICommandList());
-
-FEditorPerformanceStatusBarMenuCommands::FEditorPerformanceStatusBarMenuCommands()
-	: TCommands<FEditorPerformanceStatusBarMenuCommands>
-	(
-		"EditorPerformanceSettings",
-		LOCTEXT("Editor Performance", "Editor Performance"),
-		"LevelEditor",
-		FAppStyle::GetAppStyleSetName()
-		)
-{}
-
-void FEditorPerformanceStatusBarMenuCommands::RegisterCommands()
-{
-	UI_COMMAND(ChangeSettings, "Change Performance Settings", "Opens the Editor Performance Settings.", EUserInterfaceActionType::Button, FInputChord());
-	UI_COMMAND(ViewPerformanceReport, "View Performance Report", "Opens the Editor Performance Report panel.", EUserInterfaceActionType::Button, FInputChord());
-
-	ActionList->MapAction(
-		ChangeSettings,
-		FExecuteAction::CreateStatic(&FEditorPerformanceStatusBarMenuCommands::ChangeSettings_Clicked)
-	);
-
-	ActionList->MapAction(
-		ViewPerformanceReport,
-		FExecuteAction::CreateStatic(&FEditorPerformanceStatusBarMenuCommands::ViewPerformanceReport_Clicked)
-	);
-} 
-
-void FEditorPerformanceStatusBarMenuCommands::ChangeSettings_Clicked()
-{
-	FModuleManager::LoadModuleChecked<ISettingsModule>("Settings").ShowViewer("Editor", "General", "EditorPerformanceSettings");
-}
-
-void FEditorPerformanceStatusBarMenuCommands::ViewPerformanceReport_Clicked()
+FReply SEditorPerformanceStatusBarWidget::ViewPerformanceReport_Clicked()
 {
 	FModuleManager::LoadModuleChecked<FEditorPerformanceModule>("EditorPerformance").ShowPerformanceReportTab();
-}
 
-TSharedRef<SWidget> SEditorPerformanceStatusBarWidget::CreateStatusBarMenu()
-{
-	UToolMenu* Menu = UToolMenus::Get()->RegisterMenu("StatusBar.ToolBar.EditorPerformance", NAME_None, EMultiBoxType::Menu, false);
-
-	{
-		FToolMenuSection& Section = Menu->AddSection("PerformanceSettingsSection", LOCTEXT("PerformanceSettingsSection", "Settings"));
-
-		Section.AddMenuEntry(
-			FEditorPerformanceStatusBarMenuCommands::Get().ChangeSettings,
-			TAttribute<FText>(),
-			TAttribute<FText>(),
-			FSlateIcon(FAppStyle::GetAppStyleSetName(), "EditorPerformance.Settings")
-		);
-	}
-
-	{
-		FToolMenuSection& Section = Menu->AddSection("PerformanceReportSection", LOCTEXT("PerformanceReportSection", "Panels"));
-
-		Section.AddMenuEntry(
-			FEditorPerformanceStatusBarMenuCommands::Get().ViewPerformanceReport,
-			TAttribute<FText>(),
-			TAttribute<FText>(),
-			FSlateIcon(FAppStyle::GetAppStyleSetName(), "EditorPerformance.Report.Panel")
-		);
-	}
-
-	return UToolMenus::Get()->GenerateWidget("StatusBar.ToolBar.EditorPerformance", FToolMenuContext(FEditorPerformanceStatusBarMenuCommands::ActionList));
+	return FReply::Handled();
 }
 
 void SEditorPerformanceStatusBarWidget::Construct(const FArguments& InArgs)
 {
 	this->ChildSlot
 		[
-			SNew(SComboButton)
-			.ContentPadding(FMargin(6.0f, 0.0f))
-			.MenuPlacement(MenuPlacement_AboveAnchor)
-			.ComboButtonStyle(&FAppStyle::Get().GetWidgetStyle<FComboButtonStyle>("SimpleComboButton"))
-			.ButtonContent()
+			SNew(SButton)
+			.ButtonStyle(&FAppStyle::Get().GetWidgetStyle<FButtonStyle>("StatusBar.StatusBarButton"))
+			.Content()
 			[
-
 				SNew(SHorizontalBox)
 				+ SHorizontalBox::Slot()
 				.AutoWidth()
 				.VAlign(VAlign_Center)
-				.Padding(0, 0, 3, 0)
+				.Padding(0.f)
 				[
 					SNew(SOverlay)
 					+ SOverlay::Slot()
 					.HAlign(HAlign_Center)
 					.VAlign(VAlign_Top)
-				[
-					SNew(SImage)
-					.ColorAndOpacity_Lambda([this] { return GetStatusIconColor(); })
-					.Image_Lambda([this] { return GetStatusIcon(); })
-					.ToolTipText_Lambda([this] { return GetStatusToolTipText(); })
+					[
+						SNew(SImage)
+						.ColorAndOpacity_Lambda([this] { return GetStatusIconColor(); })
+						.Image_Lambda([this] { return GetStatusIcon(); })
+						.ToolTipText_Lambda([this] { return GetStatusToolTipText(); })
+					]
 				]
 			]
-
-			+ SHorizontalBox::Slot()
-				.AutoWidth()
-				.VAlign(VAlign_Center)
-				.Padding(0, 0, 10, 0)
-				[
-					SNew(STextBlock)
-					.Text_Lambda([this] { return GetTitleText(); })
-					.ToolTipText_Lambda([this] { return GetStatusToolTipText(); })
-				]
-	]
-	.OnGetMenuContent(FOnGetContent::CreateRaw(this, &SEditorPerformanceStatusBarWidget::CreateStatusBarMenu))
+			.OnClicked(FOnClicked::CreateStatic(&SEditorPerformanceStatusBarWidget::ViewPerformanceReport_Clicked))
 		];
 
 	RegisterActiveTimer(5.0f, FWidgetActiveTimerDelegate::CreateSP(this, &SEditorPerformanceStatusBarWidget::UpdateState));
-
-
 }
 
 EActiveTimerReturnType SEditorPerformanceStatusBarWidget::UpdateState(double InCurrentTime, float InDeltaTime)
@@ -314,19 +239,7 @@ const FSlateBrush* SEditorPerformanceStatusBarWidget::GetStatusIcon() const
 
 const FSlateColor SEditorPerformanceStatusBarWidget::GetStatusIconColor() const
 {
-	switch (EditorPerformanceState)
-	{
-		default:
-		case EEditorPerformanceState::Good:
-		{
-			return FLinearColor::White; 
-		}
-
-		case EEditorPerformanceState::Warnings:
-		{
-			return FLinearColor::White;
-		}
-	}
+	return FSlateColor::UseForeground();
 }
 
 FText SEditorPerformanceStatusBarWidget::GetStatusToolTipText() const
diff --git a/Engine/Plugins/Experimental/EditorPerformance/Source/EditorPerformance/Private/SEditorPerformanceStatusBar.h b/Engine/Plugins/Experimental/EditorPerformance/Source/EditorPerformance/Private/SEditorPerformanceStatusBar.h
--- a/Engine/Plugins/Experimental/EditorPerformance/Source/EditorPerformance/Private/SEditorPerformanceStatusBar.h
+++ b/Engine/Plugins/Experimental/EditorPerformance/Source/EditorPerformance/Private/SEditorPerformanceStatusBar.h
@@ -23,26 +23,6 @@ enum class EEditorPerformanceState : uint8
 	Warnings,
 };
 
-class FEditorPerformanceStatusBarMenuCommands : public TCommands<FEditorPerformanceStatusBarMenuCommands>
-{
-public:
-
-	FEditorPerformanceStatusBarMenuCommands();
-
-	virtual void RegisterCommands() override;
-
-private:
-
-	static void ChangeSettings_Clicked();
-	static void ViewPerformanceReport_Clicked();
-	
-public:
-
-	TSharedPtr< FUICommandInfo > ChangeSettings;
-	TSharedPtr< FUICommandInfo > ViewPerformanceReport;
-	static TSharedRef<FUICommandList> ActionList;
-};
-
 class SEditorPerformanceStatusBarWidget : public SCompoundWidget
 {
 	SLATE_BEGIN_ARGS(SEditorPerformanceStatusBarWidget) {}
@@ -51,13 +31,13 @@ class SEditorPerformanceStatusBarWidget : public SCompoundWidget
 	void Construct(const FArguments& InArgs);
 
 private:
+	static FReply ViewPerformanceReport_Clicked();
 
 	FText							GetTitleToolTipText() const;
 	FText							GetTitleText() const;
 	const FSlateBrush*				GetStatusIcon() const;
 	const FSlateColor				GetStatusIconColor() const;
 	FText							GetStatusToolTipText() const;
-	TSharedRef<SWidget>				CreateStatusBarMenu();
 	EActiveTimerReturnType			UpdateState(double InCurrentTime, float InDeltaTime);
 	
 	TSharedPtr<SNotificationItem>	NotificationItem;
diff --git a/Engine/Plugins/Experimental/EditorPerformance/Source/EditorPerformance/Public/EditorPerformanceModule.h b/Engine/Plugins/Experimental/EditorPerformance/Source/EditorPerformance/Public/EditorPerformanceModule.h
--- a/Engine/Plugins/Experimental/EditorPerformance/Source/EditorPerformance/Public/EditorPerformanceModule.h
+++ b/Engine/Plugins/Experimental/EditorPerformance/Source/EditorPerformance/Public/EditorPerformanceModule.h
@@ -59,8 +59,6 @@ public:
 	UE_API FEditorPerformanceModule::EEditorState GetEditorState() const;
 
 private:
-
-	UE_API TSharedPtr<SWidget> CreatePerformanceReportDialog();
 	UE_API TSharedRef<SDockTab> CreatePerformanceReportTab(const FSpawnTabArgs& Args);	
 	TWeakPtr<SDockTab> PerformanceReportTab;
 
diff --git a/Engine/Plugins/Experimental/EditorPerformance/Source/StallLogSubsystem/Private/StallLogSubsystem.cpp b/Engine/Plugins/Experimental/EditorPerformance/Source/StallLogSubsystem/Private/StallLogSubsystem.cpp
--- a/Engine/Plugins/Experimental/EditorPerformance/Source/StallLogSubsystem/Private/StallLogSubsystem.cpp
+++ b/Engine/Plugins/Experimental/EditorPerformance/Source/StallLogSubsystem/Private/StallLogSubsystem.cpp
@@ -168,8 +168,6 @@ void FStallLogHistory::ClearStallLog()
 
 namespace
 {
-	const FName StallLogTabName = FName(TEXT("StallLogTab"));
-
 	const FName ColumnName_Location = FName(TEXT("Location"));
 	const FName ColumnName_Thread = FName(TEXT("Thread"));
 	const FName ColumnName_Duration = FName(TEXT("Duration"));
@@ -467,93 +465,6 @@ namespace
 		];
 
 	}
-	
-	class SStallLogStatusBarWidget : public SCompoundWidget
-	{
-		SLATE_BEGIN_ARGS(SStallLogStatusBarWidget) {}
-			SLATE_ARGUMENT(FGetSlateBrush, GetFilterBadgeIcon)
-			SLATE_ARGUMENT(FGetSlateBrush, GetFilterBadgeNotificationIcon)
-		SLATE_END_ARGS()
-
-		void Construct(const FArguments& InArgs);
-
-	private:
-		
-		FText						GetToolTipText() const;
-		const FSlateBrush*			GetBadgeIcon() const;
-		const FSlateBrush*			GetBadgeNotificationIcon() const;
-
-		FGetSlateBrush GetFilterBadgeIconDelegate;
-		FGetSlateBrush GetFilterBadgeNotificationIconDelegate;
-	};
-
-	FText SStallLogStatusBarWidget::GetToolTipText() const
-	{
-		return LOCTEXT("StallLogStatusBarToolTip", "Opens the Stall Log");
-	}
-
-	const FSlateBrush* SStallLogStatusBarWidget::GetBadgeIcon() const
-	{
-		if (GetFilterBadgeIconDelegate.IsBound())
-		{
-			return GetFilterBadgeIconDelegate.Execute();
-		}
-		return nullptr;
-	}
-
-	const FSlateBrush* SStallLogStatusBarWidget::GetBadgeNotificationIcon() const
-	{
-		if (GetFilterBadgeNotificationIconDelegate.IsBound())
-		{
-			return GetFilterBadgeNotificationIconDelegate.Execute();
-		}
-		return nullptr;
-	}
-
-	void SStallLogStatusBarWidget::Construct(const FArguments& InArgs)
-	{
-		GetFilterBadgeIconDelegate = InArgs._GetFilterBadgeIcon;
-		GetFilterBadgeNotificationIconDelegate = InArgs._GetFilterBadgeNotificationIcon;
-		this->ChildSlot
-		[
-			SNew(SButton)
-			.ButtonStyle(&FAppStyle::Get().GetWidgetStyle<FButtonStyle>("StatusBar.StatusBarButton"))
-			.ContentPadding(FMargin(6.0f, 0.0f))
-			[
-				SNew(SHorizontalBox)
-				+ SHorizontalBox::Slot()
-				.AutoWidth()
-				.VAlign(VAlign_Center)
-				.Padding(0, 0, 3, 0)
-				[
-					SNew(SOverlay)
-
-					+ SOverlay::Slot()
-					.HAlign(HAlign_Center)
-					.VAlign(VAlign_Top)
-					[
-						SNew(SImage)
-						.ColorAndOpacity(FSlateColor::UseForeground())
-						.Image_Lambda([this] { return GetBadgeIcon();  })
-						.ToolTipText_Lambda([this] { return GetToolTipText(); })
-					]
-					+ SOverlay::Slot()
-					.HAlign(HAlign_Center)
-					.VAlign(VAlign_Top)
-					[
-						SNew(SImage)
-							.Image_Lambda([this] { return GetBadgeNotificationIcon();  })
-					]
-				]
-			]
-			.OnClicked_Lambda([]()
-			{
-				FGlobalTabmanager::Get()->TryInvokeTab(FTabId(StallLogTabName));
-				return FReply::Handled();
-			})
-		];
-	}
-
 }
 
 void FStallLogHistory::OnStallDetected(
@@ -611,11 +522,7 @@ void FStallLogHistory::OnStallCompleted(uint32 ThreadID, FStringView StatName, u
 void FStallLogHistory::AddStallToLog(FStallLogItemPtr LogItemToAdd)
 {
 	StallLogs.Emplace(LogItemToAdd);
-	// If the tab is open, we want to automatically say we viewed any new stalls
-	if (!FGlobalTabmanager::Get()->FindExistingLiveTab(FTabId(StallLogTabName)).IsValid())
-	{
-		SetHasViewedStalls(false);
-	}
+	SetHasViewedStalls(false);
 }
 
 DECLARE_CYCLE_STAT(
@@ -664,79 +571,30 @@ void UStallLogSubsystem::Initialize(FSubsystemCollectionBase& Collection)
 
 	RegisterStallsLogListing();
 	RegisterStallDetectedDelegates();
-
-	/** Register a tab spawner invoked by the tools bar button */
-	{
-		const FSlateIcon StallLogIcon(FAppStyle::GetAppStyleSetName(), "EditorPerformance.Stall");
-
-		FGlobalTabmanager::Get()->RegisterNomadTabSpawner(StallLogTabName, FOnSpawnTab::CreateUObject(this, &UStallLogSubsystem::CreateStallLogTab))
-			.SetDisplayName(LOCTEXT("StallLogTabTitle", "Stall Log"))
-			.SetTooltipText(LOCTEXT("StallLogTabToolTipText", "Show Stall Log"))
-			.SetGroup(WorkspaceMenu::GetMenuStructure().GetToolsCategory())
-			.SetIcon(StallLogIcon);
-	}
-	
-	/** Add the widget button to the tools bar */	
-	UToolMenu* Menu = UToolMenus::Get()->ExtendMenu("LevelEditor.StatusBar.ToolBar");
-
-	FToolMenuSection& StallDetectorSection = Menu->AddSection("StallLog", FText::GetEmpty());
-
-	auto CreateStallLogWidget = [this]()
-	{
-		return SNew(SStallLogStatusBarWidget)
-			.GetFilterBadgeIcon(FGetSlateBrush::CreateUObject(this, &UStallLogSubsystem::GetStatusBarBadgeIcon))
-			.GetFilterBadgeNotificationIcon(FGetSlateBrush::CreateUObject(this, &UStallLogSubsystem::GetStatusBarBadgeNotificationIcon));
-	};
-	
-	StallDetectorSection.AddEntry(
-		FToolMenuEntry::InitWidget("StallLogStatusBar", CreateStallLogWidget(), FText::GetEmpty(), true, false));
 }
 
 void UStallLogSubsystem::Deinitialize()
 {
-	FGlobalTabmanager::Get()->UnregisterNomadTabSpawner(StallLogTabName);
-
-	if (const TSharedPtr<SDockTab> StallLogTabShared = StallLogTab.Pin())
-	{
-		StallLogTabShared->RequestCloseTab();
-	}
-
 	UnregisterStallDetectedDelegates();
 	UnregisterStallsLogListing();
 }
 	
-TSharedRef<SDockTab> UStallLogSubsystem::CreateStallLogTab(const FSpawnTabArgs& InArgs)
+TSharedRef<SWidget> UStallLogSubsystem::CreateStallLogPanel()
 {
 	StallLogHistory->SetHasViewedStalls(true);
-	return SAssignNew(StallLogTab, SDockTab)
-	.TabRole(ETabRole::NomadTab)
-	[
-		SAssignNew(StallLog, SStallLog)
-			.StallLogItems(&StallLogHistory->GetStallLog())
-			.OnClearLog_Lambda([StallLogHistoryWkPtr = this->StallLogHistory.ToWeakPtr()]
-			{
-				TSharedPtr<FStallLogHistory> StallLogHistoryPtr = StallLogHistoryWkPtr.Pin();
-				if (StallLogHistoryPtr)
-				{
-					StallLogHistoryPtr->ClearStallLog();
-				}
-			})
-	];
-}
 
-const FSlateBrush* UStallLogSubsystem::GetStatusBarBadgeIcon() const
-{
-	return FAppStyle::GetBrush(TEXT("EditorPerformance.Stall"));
-}
-
-const FSlateBrush* UStallLogSubsystem::GetStatusBarBadgeNotificationIcon() const
-{
-	// If the StallLog has any unviewed stalls, display the notification overlay
-	if (!StallLogHistory->GetHasViewedStalls())
+	TSharedRef<SWidget> StallLog = SNew(SStallLog)
+	.StallLogItems(&StallLogHistory->GetStallLog())
+	.OnClearLog_Lambda([StallLogHistoryWkPtr = this->StallLogHistory.ToWeakPtr()]
 	{
-		return FAppStyle::GetBrush(TEXT("Icons.BadgeModified"));
-	}
-	return nullptr;
+		TSharedPtr<FStallLogHistory> StallLogHistoryPtr = StallLogHistoryWkPtr.Pin();
+		if (StallLogHistoryPtr)
+		{
+			StallLogHistoryPtr->ClearStallLog();
+		}
+	});
+
+	return StallLog;
 }
 
 void UStallLogSubsystem::RegisterStallDetectedDelegates()
diff --git a/Engine/Plugins/Experimental/EditorPerformance/Source/StallLogSubsystem/Public/StallLogSubsystem.h b/Engine/Plugins/Experimental/EditorPerformance/Source/StallLogSubsystem/Public/StallLogSubsystem.h
--- a/Engine/Plugins/Experimental/EditorPerformance/Source/StallLogSubsystem/Public/StallLogSubsystem.h
+++ b/Engine/Plugins/Experimental/EditorPerformance/Source/StallLogSubsystem/Public/StallLogSubsystem.h
@@ -32,23 +32,17 @@ public:
 	UE_API virtual void Initialize(FSubsystemCollectionBase& Collection) override;
 	UE_API virtual void Deinitialize() override;
 
-private:
-
-	UE_API TSharedRef<SDockTab> CreateStallLogTab(const FSpawnTabArgs& InArgs);
-	UE_API const FSlateBrush* GetStatusBarBadgeIcon() const;
-	UE_API const FSlateBrush* GetStatusBarBadgeNotificationIcon() const;
+	UE_API TSharedRef<SWidget> CreateStallLogPanel();
 
-	UE_API void RegisterStallDetectedDelegates();
-	UE_API void UnregisterStallDetectedDelegates();
+private:
+	void RegisterStallDetectedDelegates();
+	void UnregisterStallDetectedDelegates();
 
 private:
 
 	FDelegateHandle OnStallDetectedDelegate;
 	FDelegateHandle OnStallCompletedDelegate;
 	TSharedPtr<FStallLogHistory, ESPMode::ThreadSafe> StallLogHistory;
-	
-	TSharedPtr<SWidget> StallLog;
-	TWeakPtr<SDockTab> StallLogTab;
 };
 
 #undef UE_API
diff --git a/Engine/Source/Editor/EditorStyle/Private/StarshipStyle.cpp b/Engine/Source/Editor/EditorStyle/Private/StarshipStyle.cpp
--- a/Engine/Source/Editor/EditorStyle/Private/StarshipStyle.cpp
+++ b/Engine/Source/Editor/EditorStyle/Private/StarshipStyle.cpp
@@ -1158,11 +1158,14 @@ void FStarshipEditorStyle::FStyle::SetupGeneralStyles()
 	// Editor Performance Data
 	{
 		Set("EditorPerformance.Settings", new IMAGE_BRUSH_SVG("Starship/DerivedData/DD_Cache_Settings", Icon16x16));
-		Set("EditorPerformance.Notification.Warning", new IMAGE_BRUSH_SVG("Starship/DerivedData/DD_RemoteCache_Warning", Icon16x16, EStyleColor::Warning));
-		Set("EditorPerformance.Notification.Good", new IMAGE_BRUSH_SVG("Starship/DerivedData/DD_RemoteCache_Idle", Icon16x16, EStyleColor::Success));
+		Set("EditorPerformance.Notification.Warning", new IMAGE_BRUSH_SVG("Starship/Common/Stall", Icon16x16, EStyleColor::Warning));
+		Set("EditorPerformance.Notification.Good", new IMAGE_BRUSH_SVG("Starship/Common/Stall", Icon16x16));
 		Set("EditorPerformance.Report.Panel", new IMAGE_BRUSH_SVG("Starship/DerivedData/DD_Cache_Statistics", Icon16x16));
 		Set("EditorPerformance.Report.Warning", new CORE_IMAGE_BRUSH_SVG("Starship/Common/alert-triangle", Icon16x16, EStyleColor::Warning));
 		Set("EditorPerformance.Stall", new IMAGE_BRUSH_SVG("Starship/Common/Stall", Icon16x16));
+
+		Set("EditorDiagnostics.SetShowActivityMonitor", new IMAGE_BRUSH_SVG("Starship/Common/Scalability_16", Icon20x20));
+		Set("EditorDiagnostics.SetShowStalls", new IMAGE_BRUSH_SVG("Starship/Common/Stall", Icon20x20));
 	}
 
 	// Scene Outliner
