Enable Editor Performance Tools by default
- Deprecated NotificationList and replaced it with NotifyList since NotificationList had been programmatically saved to the user configs.
- Made it so that we only show a bad state in the status bar for KPIs that the user wants to be notified about since some KPIs are quite flaky.
- UX pass on the performance panel.
- Swapped the status bar icon

diff --git a/Engine/Config/BaseEditorSettings.ini b/Engine/Config/BaseEditorSettings.ini
--- a/Engine/Config/BaseEditorSettings.ini
+++ b/Engine/Config/BaseEditorSettings.ini
@@ -69,5 +69,7 @@ TriggerPressedThreshold_Rift=0.500000
 Duration=(SubFrame=0.0)
 
 [/Script/UnrealEd.EditorPerformanceSettings]
-NotificationList=Hardware_Memory_Pressure
+bEnableEditorPerformanceTool=True
+bEnableNotifications=True
+NotifyList=Memory_Memory_Pressure
 
diff --git a/Engine/Plugins/Experimental/EditorPerformance/Source/EditorPerformance/Private/SEditorPerformanceDialogs.cpp b/Engine/Plugins/Experimental/EditorPerformance/Source/EditorPerformance/Private/SEditorPerformanceDialogs.cpp
--- a/Engine/Plugins/Experimental/EditorPerformance/Source/EditorPerformance/Private/SEditorPerformanceDialogs.cpp
+++ b/Engine/Plugins/Experimental/EditorPerformance/Source/EditorPerformance/Private/SEditorPerformanceDialogs.cpp
@@ -11,8 +11,10 @@
 #include "Widgets/SBoxPanel.h"
 #include "Widgets/SToolTip.h"
 #include "Widgets/Images/SImage.h"
+#include "Widgets/Input/SSegmentedControl.h"
 #include "Widgets/Layout/SGridPanel.h"
 #include "Widgets/Layout/SScrollBox.h"
+#include "Widgets/Layout/SSeparator.h"
 #include "Widgets/Text/STextBlock.h"
 #include "Widgets/Input/SComboBox.h"
 #include "Widgets/Input/SCheckBox.h"
@@ -22,11 +24,24 @@
 #include "ISettingsModule.h"
 #include "Editor/EditorPerformanceSettings.h"
 #include "SEnumCombo.h"
+#include "ToolMenus.h"
+
 
 #include UE_INLINE_GENERATED_CPP_BY_NAME(SEditorPerformanceDialogs)
 
 #define LOCTEXT_NAMESPACE "EditorPerformance"
 
+
+const TCHAR* SEditorPerformanceReportDialog::SettingsMenuName = TEXT("EditorPerfSettings");
+
+SEditorPerformanceReportDialog::~SEditorPerformanceReportDialog()
+{
+	if (UToolMenus* ToolMenus = UToolMenus::Get())
+	{
+		ToolMenus->RemoveMenu(SettingsMenuName);
+	}
+}
+
 void SEditorPerformanceReportDialog::Construct(const FArguments& InArgs)
 {
 	const float RowMargin = 0.0f;
@@ -34,7 +49,6 @@ void SEditorPerformanceReportDialog::Construct(const FArguments& InArgs)
 	const float ColumnMargin = 10.0f;
 	const FMargin DefaultMarginFirstColumn(ColumnMargin, RowMargin);
 	const FSlateColor TitleColor = FStyleColors::AccentWhite;
-	const FSlateFontInfo TitleFont = FCoreStyle::GetDefaultFontStyle("Bold", 10);
 
 	FEditorPerformanceModule& EditorPerfModule = FModuleManager::LoadModuleChecked<FEditorPerformanceModule>("EditorPerformance");
 
@@ -43,22 +57,20 @@ void SEditorPerformanceReportDialog::Construct(const FArguments& InArgs)
 		SNew(SScrollBox)
 		.Orientation(Orient_Vertical)
 		+ SScrollBox::Slot()
+		.HAlign(HAlign_Fill)
 		[
 			SNew(SVerticalBox)
 			+ SVerticalBox::Slot()
-			.Padding(0, 20, 0, 0)
+			.Padding(0, 5, 0, 0)
 			.AutoHeight()
 			[
-				SNew(SHorizontalBox)
-				+SHorizontalBox::Slot()
-				.FillWidth(1.0f)
+				SNew(SBox)
 				[
 					SNew(STextBlock)
-					.Margin(TitleMargin)
+					.Margin(FMargin(TitleMargin, TitleMargin, TitleMargin, 0.f))
 					.ColorAndOpacity(TitleColor)
-					.Font(TitleFont)
 					.Justification(ETextJustify::Left)
-					.Text_Lambda([this,&EditorPerfModule]
+					.Text_Lambda([this, &EditorPerfModule]
 						{ 
 							return FText::Format(LOCTEXT("Profile", "Profile: {0}"), FText::FromString(*EditorPerfModule.GetKPIProfileName())); 
 						}
@@ -67,7 +79,7 @@ void SEditorPerformanceReportDialog::Construct(const FArguments& InArgs)
 			]
 			+ SVerticalBox::Slot()
 			.AutoHeight()
-			.Padding(0, 5, 0, 0)
+			.Padding(0.f)
 			.Expose(KPIGridSlot)
 			[
 				GetKPIGridPanel()
@@ -79,26 +91,12 @@ void SEditorPerformanceReportDialog::Construct(const FArguments& InArgs)
 			[
 				GetHintGridPanel()
 			]
-			+ SVerticalBox::Slot()
-			.AutoHeight()
-			.Padding(0, 5, 0, 0)
-			.Expose(SettingsGridSlot)
-			[
-				GetSettingsGridPanel()
-			]
 		]
 	];
-
-	RegisterActiveTimer(5.0f, FWidgetActiveTimerDelegate::CreateSP(this, &SEditorPerformanceReportDialog::UpdateGridPanels));
 }
 
 EActiveTimerReturnType SEditorPerformanceReportDialog::UpdateGridPanels(double InCurrentTime, float InDeltaTime)
 {
-	(*SettingsGridSlot)
-	[
-		GetSettingsGridPanel()
-	];
-
 	(*KPIGridSlot)
 	[
 		GetKPIGridPanel()
@@ -114,207 +112,6 @@ EActiveTimerReturnType SEditorPerformanceReportDialog::UpdateGridPanels(double I
 	return EActiveTimerReturnType::Continue;
 }
 
-TSharedRef<SWidget> SEditorPerformanceReportDialog::GetSettingsGridPanel()
-{
-	TSharedRef<SGridPanel> Panel =
-		SNew(SGridPanel);
-
-	const float RowMargin = 0.0f;
-	const float ColumnMargin = 10.0f;
-	const FSlateColor TitleColor = FStyleColors::AccentWhite;
-	const FSlateFontInfo TitleFont = FCoreStyle::GetDefaultFontStyle("Bold", 10);
-
-	const FMargin TitleMargin(0.0f, 10.0f, ColumnMargin, 10.0f);
-	const FMargin TitleMarginFirstColumn(ColumnMargin, 10.0f);
-	const FMargin DefaultMargin(0.0f, RowMargin, ColumnMargin, RowMargin);
-	const FMargin DefaultMarginFirstColumn(ColumnMargin, RowMargin);
-
-	int32 Row = 0;
-
-	Panel->AddSlot(0, Row)
-		.HAlign(HAlign_Left)
-		[
-			SNew(STextBlock)
-			.Margin(TitleMarginFirstColumn)
-			.ColorAndOpacity(TitleColor)
-			.Font(TitleFont)
-			.Text(LOCTEXT("SettingsText", "Settings"))
-		];
-
-	Row++;
-
-	Panel->AddSlot(0, Row)
-		.HAlign(HAlign_Left)
-		[
-			SNew(SCheckBox)
-			.IsChecked_Lambda([]
-				{
-					const UEditorPerformanceSettings* EditorPerformanceSettings = GetDefault<UEditorPerformanceSettings>();
-					return EditorPerformanceSettings && EditorPerformanceSettings->bEnableNotifications ? ECheckBoxState::Checked : ECheckBoxState::Unchecked;
-				})
-			.OnCheckStateChanged_Lambda([this](ECheckBoxState NewState)
-			{
-				UEditorPerformanceSettings* EditorPerformanceSettings = GetMutableDefault<UEditorPerformanceSettings>();
-
-				if (EditorPerformanceSettings)
-				{
-					EditorPerformanceSettings->bEnableNotifications = NewState == ECheckBoxState::Checked;
-					EditorPerformanceSettings->PostEditChange();
-					EditorPerformanceSettings->SaveConfig();
-				}
-
-				UpdateGridPanels(0.0f, 0.0f);
-			})
-			.Padding(FMargin(4.0f, 0.0f))
-			[
-				SNew(STextBlock)
-				.Text(LOCTEXT("EnableNotificationsText", "Notifications"))
-				.ToolTipText(LOCTEXT("EnableNotificationsToolTip", "Enable All Notifications"))
-				.ColorAndOpacity(EStyleColor::Foreground)
-			]
-		];
-
-	Panel->AddSlot(1, Row)
-		.HAlign(HAlign_Left)
-		[
-			SNew(SCheckBox)
-			.IsChecked_Lambda([]
-				{
-					const UEditorPerformanceSettings* EditorPerformanceSettings = GetDefault<UEditorPerformanceSettings>();
-					return EditorPerformanceSettings && EditorPerformanceSettings->bEnableSnapshots ? ECheckBoxState::Checked : ECheckBoxState::Unchecked;
-				})
-			.OnCheckStateChanged_Lambda([this](ECheckBoxState NewState)
-				{
-					UEditorPerformanceSettings* EditorPerformanceSettings = GetMutableDefault<UEditorPerformanceSettings>();
-
-					if (EditorPerformanceSettings)
-					{
-						EditorPerformanceSettings->bEnableSnapshots = NewState == ECheckBoxState::Checked;
-						EditorPerformanceSettings->PostEditChange();
-						EditorPerformanceSettings->SaveConfig();
-					}
-
-					UpdateGridPanels(0.0f, 0.0f);
-				})
-				.Padding(FMargin(4.0f, 0.0f))
-				[
-					SNew(STextBlock)
-					.Text(LOCTEXT("EnableSnapshotsText", "Snapshots"))
-					.ToolTipText(LOCTEXT("EnableSnapshotsToolTip", "Enable Automatic Capture of Unreal Insights Snaphsots"))
-					.ColorAndOpacity(EStyleColor::Foreground)
-				]
-			];
-
-	Panel->AddSlot(2, Row)
-		.HAlign(HAlign_Left)
-		[
-			SNew(SCheckBox)
-			.IsChecked_Lambda([]
-				{
-					const UEditorPerformanceSettings* EditorPerformanceSettings = GetDefault<UEditorPerformanceSettings>();
-					return EditorPerformanceSettings && EditorPerformanceSettings->bEnableTelemetry ? ECheckBoxState::Checked : ECheckBoxState::Unchecked;
-				})
-			.OnCheckStateChanged_Lambda([this](ECheckBoxState NewState)
-			{
-				UEditorPerformanceSettings* EditorPerformanceSettings = GetMutableDefault<UEditorPerformanceSettings>();
-
-				if (EditorPerformanceSettings)
-				{
-					EditorPerformanceSettings->bEnableTelemetry = NewState == ECheckBoxState::Checked;
-					EditorPerformanceSettings->PostEditChange();
-					EditorPerformanceSettings->SaveConfig();
-				}
-
-				UpdateGridPanels(0.0f, 0.0f);
-			})
-			.Padding(FMargin(4.0f, 0.0f))
-			[
-				SNew(STextBlock)
-				.Text(LOCTEXT("EnableTelemetryText", "Telemetry"))
-				.ToolTipText(LOCTEXT("EnableTelemetryToolTip", "Record Warning Telemetry Events To Analytics System"))
-				.ColorAndOpacity(EStyleColor::Foreground)
-			]
-		];
-
-	Panel->AddSlot(3, Row)
-		.HAlign(HAlign_Left)
-		[
-			SNew(SCheckBox)
-			.IsChecked_Lambda([]
-				{
-					const UEditorPerformanceSettings* EditorPerformanceSettings = GetDefault<UEditorPerformanceSettings>();
-					return EditorPerformanceSettings && EditorPerformanceSettings->bThrottleCPUWhenNotForeground ? ECheckBoxState::Checked : ECheckBoxState::Unchecked;
-				})
-			.OnCheckStateChanged_Lambda([this](ECheckBoxState NewState)
-			{
-				UEditorPerformanceSettings* EditorPerformanceSettings = GetMutableDefault<UEditorPerformanceSettings>();
-
-				if (EditorPerformanceSettings)
-				{
-					EditorPerformanceSettings->bThrottleCPUWhenNotForeground = NewState == ECheckBoxState::Checked;
-					EditorPerformanceSettings->PostEditChange();
-					EditorPerformanceSettings->SaveConfig();
-				}
-
-				UpdateGridPanels(0.0f, 0.0f);
-			})
-			.Padding(FMargin(4.0f, 0.0f))
-			[
-				SNew(STextBlock)
-				.Text(LOCTEXT("EnableBackgroundThrottlingText", "Throttling"))
-				.ToolTipText(LOCTEXT("EnableBackgroundThrottlingToolTip", "Enable CPU throttling when the Editor is in the background."))
-				.ColorAndOpacity(EStyleColor::Foreground)
-			]
-		];
-
-	Panel->AddSlot(4, Row)
-		.HAlign(HAlign_Left)
-		[
-			SNew(SCheckBox)
-				.IsChecked_Lambda([]
-					{
-						const UEditorPerformanceSettings* EditorPerformanceSettings = GetDefault<UEditorPerformanceSettings>();
-						return EditorPerformanceSettings && EditorPerformanceSettings->bShowFrameRateAndMemory ? ECheckBoxState::Checked : ECheckBoxState::Unchecked;
-					})
-				.OnCheckStateChanged_Lambda([this](ECheckBoxState NewState)
-					{
-						UEditorPerformanceSettings* EditorPerformanceSettings = GetMutableDefault<UEditorPerformanceSettings>();
-
-						if (EditorPerformanceSettings)
-						{
-							EditorPerformanceSettings->bShowFrameRateAndMemory = NewState == ECheckBoxState::Checked;
-							EditorPerformanceSettings->PostEditChange();
-							EditorPerformanceSettings->SaveConfig();
-						}
-
-						UpdateGridPanels(0.0f, 0.0f);
-					})
-						.Padding(FMargin(4.0f, 0.0f))
-						[
-							SNew(STextBlock)
-								.Text(LOCTEXT("EnableShowFrameRateAndMemoryText", "Diagnostics"))
-								.ToolTipText(LOCTEXT("EnableShowFrameRateAndMemoryToolTip", "Show the Frame Rate, Memory and Stalls."))
-								.ColorAndOpacity(EStyleColor::Foreground)
-						]
-		];
-
-	Panel->AddSlot(5, Row)
-		.HAlign(HAlign_Left)
-		.Padding(FMargin(10.0f, 10.0f))
-		[
-			SNew(SButton)
-			.Text(LOCTEXT("OpenSettingsText", "All Settings"))
-			.ToolTipText(LOCTEXT("OpenSettingsToolTip", "Open the Editor Performance Settings Tab."))
-			.OnClicked_Lambda([this]()
-			{
-				FModuleManager::LoadModuleChecked<ISettingsModule>("Settings").ShowViewer("Editor", "General", "EditorPerformanceSettings");
-				return FReply::Handled();
-			})
-		];
-
-	return Panel;
-}
-
 TSharedRef<SWidget> SEditorPerformanceReportDialog::GetHintGridPanel()
 {
 	FEditorPerformanceModule& EditorPerfModule = FModuleManager::LoadModuleChecked<FEditorPerformanceModule>("EditorPerformance");
@@ -439,126 +236,150 @@ TSharedRef<SWidget> SEditorPerformanceReportDialog::GetHintGridPanel()
 	return Panel;
 }
 
+TSharedRef<SWidget> CreateGridSlotHeaderWidget(FText SlotText, FMargin Margin, ETextJustify::Type TextJustify = ETextJustify::Left, EVisibility SeparatorVisibility = EVisibility::Visible)
+{
+	return SNew(SBorder)
+	.BorderImage(FAppStyle::Get().GetBrush("WhiteBrush"))
+	.BorderBackgroundColor(EStyleColor::Header)
+	.Padding(2.f, 0.f)
+	[
+		SNew(SBox)
+		.MinDesiredHeight(24.f)
+		[
+			SNew(SHorizontalBox)
+			+SHorizontalBox::Slot()
+			.VAlign(VAlign_Center)
+			.FillWidth(1.f)
+			[
+				SNew(STextBlock)
+				.Margin(Margin)
+				.ColorAndOpacity(EStyleColor::Foreground)
+				.Text(SlotText)
+				.Justification(TextJustify)
+			]
+			+SHorizontalBox::Slot()
+			.AutoWidth()
+			[
+				SNew(SSeparator)
+				.Thickness(1.f)
+				.Orientation(Orient_Vertical)
+				.Visibility(SeparatorVisibility)
+			]
+		]
+	];
+}
+
+TSharedRef<SWidget> CreateGridSlotWidget(TAttribute<FText> SlotText, FMargin Margin, ETextJustify::Type TextJustify = ETextJustify::Left)
+{
+	return SNew(SBorder)
+	.BorderImage(FAppStyle::Get().GetBrush("WhiteBrush"))
+	.BorderBackgroundColor(EStyleColor::Background)
+	[
+		SNew(SBorder)
+		.BorderImage(FAppStyle::Get().GetBrush("NoBorder"))
+		.VAlign(VAlign_Center)
+		[
+			SNew(STextBlock)
+			.Margin(Margin)
+			.ColorAndOpacity(EStyleColor::Foreground)
+			.Text(SlotText)
+			.Justification(TextJustify)
+		]
+	];
+}
+
 TSharedRef<SWidget> SEditorPerformanceReportDialog::GetKPIGridPanel()
 {
+	enum EGridColumns
+	{
+		GridColumns_Name,
+		GridColumns_CurrentValue,
+		GridColumns_ExpectedValue,
+		GridColumns_Failures,
+		GridColumns_WarningIcon,
+		GridColumns_Notify,
+		GridColumns_Fill,
+
+		GridColumns_Count
+	};
+
 	TSharedRef<SGridPanel> Panel =
-		SNew(SGridPanel);
+		SNew(SGridPanel)
+		.FillColumn(GridColumns_Fill, 1.f);
 
 	const float RowMargin = 0.0f;
 	const float ColumnMargin = 10.0f;
-	const FSlateColor TitleColor = FStyleColors::AccentWhite;
-	const FSlateFontInfo TitleFont = FCoreStyle::GetDefaultFontStyle("Bold", 10);
 
-	const FMargin TitleMargin(0.0f, 10.0f, ColumnMargin, 10.0f);
-	const FMargin TitleMarginFirstColumn(ColumnMargin, 10.0f);
-	const FMargin DefaultMargin(0.0f, RowMargin, ColumnMargin, RowMargin);
-	const FMargin DefaultMarginFirstColumn(ColumnMargin, RowMargin);
+	const FMargin TitleMargin(ColumnMargin, 10.0f);
+	const FMargin DefaultMargin(ColumnMargin, RowMargin);
 
 	const UEditorPerformanceSettings* EditorPerformanceSettings = GetDefault<UEditorPerformanceSettings>();
 	FEditorPerformanceModule& EditorPerfModule = FModuleManager::LoadModuleChecked<FEditorPerformanceModule>("EditorPerformance");
 
-	const bool EnableNotifcations = EditorPerformanceSettings && EditorPerformanceSettings->bEnableNotifications;
-	const bool ShowWarningsOnly = EditorPerformanceSettings && EditorPerformanceSettings->bShowWarningsOnly;
+	const bool bEnableNotifications = EditorPerformanceSettings && EditorPerformanceSettings->bEnableNotifications;
+	const bool bShowWarningsOnly = EditorPerformanceSettings && EditorPerformanceSettings->bShowWarningsOnly;
+	
 	
 	int32 Row = 0;
 
-	Panel->AddSlot(0, Row)
-		.HAlign(HAlign_Left)
-		[
-			SNew(STextBlock)
-			.Margin(TitleMarginFirstColumn)
-			.ColorAndOpacity(TitleColor)
-			.Font(TitleFont)
-			.Text(LOCTEXT("Measurements", "Measurements"))
-		];
-
-	Row++;
-
-	Panel->AddSlot(0, Row)
-		.HAlign(HAlign_Left)
+	Panel->AddSlot(GridColumns_Name, Row)
+	.HAlign(HAlign_Fill)
+	.ColumnSpan(GridColumns_Count)
+	[
+		SNew(SHorizontalBox)
+		+SHorizontalBox::Slot()
+		.AutoWidth()
+		.Padding(TitleMargin)
 		[
-			SNew(SEnumComboBox, StaticEnum<EEditorPerformanceFilterOptions>())
-			.CurrentValue_Lambda([]() -> int32
-			{
-				UEditorPerformanceSettings* EditorPerformanceSettings = GetMutableDefault<UEditorPerformanceSettings>();
-				if (EditorPerformanceSettings)
-				{
-					return static_cast<int32>(EditorPerformanceSettings->bShowWarningsOnly ? EEditorPerformanceFilterOptions::WarningsOnly : EEditorPerformanceFilterOptions::ShowAll);
-				}
-				else
-				{
-					return static_cast<int32>(EEditorPerformanceFilterOptions::ShowAll);
-				}
-			})
-			.OnEnumSelectionChanged_Lambda([this](int32 OptionValue, ESelectInfo::Type)
+			SNew(SSegmentedControl<EEditorPerformanceFilterOptions>)
+			.Value(EditorPerformanceSettings->bShowWarningsOnly ? EEditorPerformanceFilterOptions::WarningsOnly : EEditorPerformanceFilterOptions::ShowAll)
+			.OnValueChanged_Lambda([this](EEditorPerformanceFilterOptions OptionValue)
 			{
 				UEditorPerformanceSettings* EditorPerformanceSettings = GetMutableDefault<UEditorPerformanceSettings>();
 
 				if (EditorPerformanceSettings)
 				{
-					EditorPerformanceSettings->bShowWarningsOnly = ((EEditorPerformanceFilterOptions)OptionValue == EEditorPerformanceFilterOptions::WarningsOnly);
+					EditorPerformanceSettings->bShowWarningsOnly = (OptionValue == EEditorPerformanceFilterOptions::WarningsOnly);
 					EditorPerformanceSettings->PostEditChange();
 					EditorPerformanceSettings->SaveConfig();
 				}
 
 				UpdateGridPanels(0.0f, 0.0f);
 			})
-		];
 
-	Row++;
+			+ SSegmentedControl<EEditorPerformanceFilterOptions>::Slot(EEditorPerformanceFilterOptions::ShowAll)
+			.Text(LOCTEXT("ShowAll", "Show All"))
 
-	Panel->AddSlot(1, Row)
-		.HAlign(HAlign_Left)
-		[
-			SNew(STextBlock)
-			.Margin(DefaultMargin)
-			.ColorAndOpacity(EStyleColor::Foreground)
-			.Font(TitleFont)
-			.Text(LOCTEXT("CurrentValueColumn", "Current"))
-		];
-
-	Panel->AddSlot(3, Row)
-		.HAlign(HAlign_Left)
-		[
-			SNew(STextBlock)
-			.Margin(DefaultMargin)
-			.ColorAndOpacity(EStyleColor::Foreground)
-			.Font(TitleFont)
-			.Text(LOCTEXT("ExpectedValueColumn", "Expected"))
-		];
-
-	Panel->AddSlot(4, Row)
-		.HAlign(HAlign_Left)
+			+ SSegmentedControl<EEditorPerformanceFilterOptions>::Slot(EEditorPerformanceFilterOptions::WarningsOnly)
+			.Text(LOCTEXT("Warnings Only", "Warnings Only"))
+		]
+		+SHorizontalBox::Slot()
+		.HAlign(HAlign_Right)
+		.FillWidth(1.f)
+		.Padding(TitleMargin)
 		[
-			SNew(STextBlock)
-				.Margin(DefaultMargin)
-				.ColorAndOpacity(EStyleColor::Foreground)
-				.Font(TitleFont)
-				.Text(LOCTEXT("FailedValueColumn", "Failures"))
-		];
-
-	if (EnableNotifcations)
-	{
-		Panel->AddSlot(6, Row)
-			.HAlign(HAlign_Center)
+			SNew(SComboButton)
+			.ComboButtonStyle(FAppStyle::Get(), "SimpleComboButtonWithIcon")
+			.HasDownArrow(false)
+			.OnGetMenuContent( this, &SEditorPerformanceReportDialog::CreateSettingsMenuWidget )
+			.ButtonContent()
 			[
-				SNew(STextBlock)
-				.Margin(DefaultMargin)
-				.ColorAndOpacity(EStyleColor::Foreground)
-				.Font(TitleFont)
-				.Text(LOCTEXT("NotifyColumn", "Notify"))
-			];
-	}
+				SNew(SImage)
+				.Image(FAppStyle::GetBrush("Icons.Settings"))
+				.ColorAndOpacity(FSlateColor::UseForeground())
+			]
+		]
+	];
 
 	Row++;
-	
+
 	TMap<FName, TArray<FKPIValue>> SortedKPIValues;
 
 	for (FKPIValues::TConstIterator It(EditorPerfModule.GetKPIRegistry().GetKPIValues()); It; ++It)
 	{
 		const FKPIValue& KPIValue = It->Value;
 
-		if (ShowWarningsOnly && KPIValue.GetState()!=FKPIValue::Bad )
+		if (bShowWarningsOnly && KPIValue.GetState() != FKPIValue::Bad )
 		{
 			continue;
 		}
@@ -586,160 +407,243 @@ TSharedRef<SWidget> SEditorPerformanceReportDialog::GetKPIGridPanel()
 			CategoryText = KPIValues[0].DisplayCategory;
 		}
 
+		int32 NameColumnSpan = Row == 1 ? 1 : GridColumns_Count;
+		EVisibility NameSeparatorVisibility = NameColumnSpan == 1 ? EVisibility::Visible : EVisibility::Hidden;
+
 		// Render the category name
-		Panel->AddSlot(0, Row)
-			.HAlign(HAlign_Left)
+		Panel->AddSlot(GridColumns_Name, Row)
+		.HAlign(HAlign_Fill)
+		.ColumnSpan(NameColumnSpan)
+		[
+			CreateGridSlotHeaderWidget(CategoryText, DefaultMargin, ETextJustify::Left, NameSeparatorVisibility)
+		];
+
+		if (NameColumnSpan == 1)
+		{
+			Panel->AddSlot(GridColumns_CurrentValue, Row)
+			.HAlign(HAlign_Fill)
+			.VAlign(VAlign_Center)
 			[
-				SNew(STextBlock)
-				.Margin(TitleMarginFirstColumn)
-				.ColorAndOpacity(EStyleColor::Foreground)
-				.Font(TitleFont)
-				.Text(CategoryText)
+				CreateGridSlotHeaderWidget(LOCTEXT("CurrentValueColumn", "Current"), DefaultMargin)
+			];
+
+			Panel->AddSlot(GridColumns_ExpectedValue, Row)
+			.HAlign(HAlign_Fill)
+			.VAlign(VAlign_Center)
+			[
+				CreateGridSlotHeaderWidget(LOCTEXT("ExpectedValueColumn", "Expected"), DefaultMargin)
+			];
+
+			Panel->AddSlot(GridColumns_Failures, Row)
+			.ColumnSpan(2)
+			.HAlign(HAlign_Fill)
+			.VAlign(VAlign_Center)
+			[
+				CreateGridSlotHeaderWidget(LOCTEXT("FailedValueColumn", "Failures"), DefaultMargin)
+			];
+
+			Panel->AddSlot(GridColumns_Notify, Row)
+			.HAlign(HAlign_Fill)
+			.VAlign(VAlign_Center)
+			[
+				CreateGridSlotHeaderWidget(LOCTEXT("NotifyColumn", "Notify"), DefaultMargin, ETextJustify::Center)
 			];
-		
+	
+			Panel->AddSlot(GridColumns_Fill, Row)
+			.HAlign(HAlign_Fill)
+			.VAlign(VAlign_Center)
+			[
+				CreateGridSlotHeaderWidget(FText(), DefaultMargin, ETextJustify::Left, EVisibility::Hidden)
+			];
+		}
+
 		Row++;
 
 		for (const FKPIValue& KPIValue : KPIValues)
 		{
 			const FKPIValue::EState KPIValueState = KPIValue.GetState();
 
-			const FSlateColor KPIColor = KPIValueState==FKPIValue::Bad ? EStyleColor::Warning : EStyleColor::Foreground;
-			const FSlateBrush* KPIWarningIcon = FAppStyle::Get().GetBrush("EditorPerformance.Report.Warning");
+			const FSlateColor KPIColor = KPIValueState == FKPIValue::Bad ? EStyleColor::Warning : EStyleColor::Foreground;
 			const float KPIIconSize = 8.0f;
 			const FText& KPIDisplayName = KPIValue.DisplayName;
 			const FName& KPIPath = KPIValue.Path;
 
-			Panel->AddSlot(0, Row)
-				.HAlign(HAlign_Left)
-				.VAlign(VAlign_Center)
+			Panel->AddSlot(GridColumns_Name, Row)
+			.HAlign(HAlign_Fill)
+			.VAlign(VAlign_Fill)
+			[
+				CreateGridSlotWidget(KPIDisplayName, DefaultMargin)
+			];
+
+			Panel->AddSlot(GridColumns_CurrentValue, Row)
+			.HAlign(HAlign_Fill)
+			.VAlign(VAlign_Fill)
+			[
+				CreateGridSlotWidget(TAttribute<FText>::CreateLambda([KPIId = KPIValue.Id]()
+					{
+						FEditorPerformanceModule& EditorPerfModule = FModuleManager::LoadModuleChecked<FEditorPerformanceModule>("EditorPerformance");
+
+						FKPIValue KPIValue;
+						if (EditorPerfModule.GetKPIRegistry().GetKPIValue(KPIId, KPIValue))
+						{
+							if (KPIValue.State == FKPIValue::NotSet)
+							{
+								return LOCTEXT("PendingValue", "...");
+							}
+							else
+							{
+								return FText::FromString(*FKPIValue::GetValueAsString(KPIValue.CurrentValue, KPIValue.DisplayType, KPIValue.CustomDisplayValueGetter));
+							}
+						}
+						else
+						{
+							return FText();
+						}
+					}), DefaultMargin)
+			];
+
+			if (KPIValue.ThresholdValue)
+			{
+				FText ExpectedText = FText::FromString(*(FKPIValue::GetComparisonAsString(KPIValue.Compare) + FKPIValue::GetValueAsString(KPIValue.ThresholdValue.GetValue(), KPIValue.DisplayType, KPIValue.CustomDisplayValueGetter)));
+
+				Panel->AddSlot(GridColumns_ExpectedValue, Row)
+				.HAlign(HAlign_Fill)
+				.VAlign(VAlign_Fill)
 				[
-					SNew(STextBlock)
-					.Margin(DefaultMarginFirstColumn)
-					.ColorAndOpacity(EStyleColor::Foreground)
-					.Text(KPIDisplayName)
+					CreateGridSlotWidget(ExpectedText, DefaultMargin)
 				];
 
-			if (KPIValueState != FKPIValue::NotSet)
-			{
-				Panel->AddSlot(1, Row)
-					.HAlign(HAlign_Left)
+				Panel->AddSlot(GridColumns_Failures, Row)
+				.HAlign(HAlign_Fill)
+				.VAlign(VAlign_Fill)
+				[
+					CreateGridSlotWidget(TAttribute<FText>::CreateLambda([KPIId = KPIValue.Id]()
+					{
+						FEditorPerformanceModule& EditorPerfModule = FModuleManager::LoadModuleChecked<FEditorPerformanceModule>("EditorPerformance");
+
+						FKPIValue KPIValue;
+						if (EditorPerfModule.GetKPIRegistry().GetKPIValue(KPIId, KPIValue))
+						{
+							return FText::FromString(*FString::Printf(TEXT("%d"), KPIValue.FailureCount));
+						}
+						else
+						{
+							return FText();
+						}
+
+					}), DefaultMargin)
+				];
+
+				Panel->AddSlot(GridColumns_WarningIcon, Row)
+				.HAlign(HAlign_Fill)
+				.VAlign(VAlign_Fill)
+				[
+					SNew(SBorder)
+					.BorderImage(FCoreStyle::Get().GetBrush("WhiteBrush"))
+					.BorderBackgroundColor(EStyleColor::Background)
+					.Padding(DefaultMargin)
+					.HAlign(HAlign_Center)
 					.VAlign(VAlign_Center)
 					[
-						SNew(STextBlock)
-						.Margin(DefaultMargin)
-						.ColorAndOpacity(KPIColor)
-						.Text(FText::FromString(*FKPIValue::GetValueAsString(KPIValue.CurrentValue, KPIValue.DisplayType, KPIValue.CustomDisplayValueGetter)))
-					];
-
-				if (KPIValue.ThresholdValue)
-				{
-					Panel->AddSlot(2, Row)
-						.HAlign(HAlign_Left)
-						.VAlign(VAlign_Center)
-						[
-							SNew(STextBlock)
-							.Margin(DefaultMargin)
-							.ColorAndOpacity(KPIColor)
-							.Text(FText::FromString(*FKPIValue::GetComparisonAsString(KPIValue.Compare)))
-						];
-
-					Panel->AddSlot(3, Row)
-						.HAlign(HAlign_Left)
-						.VAlign(VAlign_Center)
-						[
-							SNew(STextBlock)
-							.Margin(DefaultMargin)
-							.ColorAndOpacity(KPIColor)
-							.Text(FText::FromString(*FKPIValue::GetValueAsString(KPIValue.ThresholdValue.GetValue(), KPIValue.DisplayType, KPIValue.CustomDisplayValueGetter)))
-						];
-
-					Panel->AddSlot(4, Row)
-						.HAlign(HAlign_Left)
-						.VAlign(VAlign_Center)
-						[
-							SNew(STextBlock)
-								.Margin(DefaultMargin)
-								.ColorAndOpacity(KPIColor)
-								.Text(FText::FromString(*FString::Printf(TEXT("%d"), KPIValue.FailureCount)))
-						];
+						SNew(SImage)
+						.Image_Lambda([KPIId = KPIValue.Id]()
+						{
+							FEditorPerformanceModule& EditorPerfModule = FModuleManager::LoadModuleChecked<FEditorPerformanceModule>("EditorPerformance");
 
-					if (KPIValueState != FKPIValue::Good)
-					{
-						Panel->AddSlot(5, Row)
-						.HAlign(HAlign_Center)
-						.VAlign(VAlign_Center)
-						[
-							SNew(SBorder)
-							.BorderImage(FAppStyle::GetBrush("NoBorder"))
-							.Padding(DefaultMargin)
-							.HAlign(HAlign_Center)
-							.VAlign(VAlign_Center)
-							[
-								SNew(SImage)
-								.Image(KPIWarningIcon)
-							]
-						];
-					}
-				}
+							FKPIValue KPIValue;
+							if (EditorPerfModule.GetKPIRegistry().GetKPIValue(KPIId, KPIValue))
+							{
+								return KPIValue.State == FKPIValue::Bad ? FAppStyle::Get().GetBrush("EditorPerformance.Report.Warning") : FAppStyle::GetNoBrush();
+							}
+							else
+							{
+								return FAppStyle::GetNoBrush();
+							}
+						})
+					]
+				];
 			}
 			else
 			{
-				Panel->AddSlot(1, Row)
-					.HAlign(HAlign_Left)
-					[
-						SNew(STextBlock)
-						.Margin(DefaultMargin)
-						.ColorAndOpacity(KPIColor)
-						.Text(LOCTEXT("PendingValue", "..."))
-					];
+				Panel->AddSlot(GridColumns_ExpectedValue, Row)
+				.ColumnSpan(3)
+				.HAlign(HAlign_Fill)
+				.VAlign(VAlign_Fill)
+				[
+					CreateGridSlotWidget(FText(), DefaultMargin)
+				];
 			}
 
-			if (EnableNotifcations && KPIValue.ThresholdValue)
+			if (bEnableNotifications && KPIValue.ThresholdValue)
 			{
-				Panel->AddSlot(6, Row)
-					.HAlign(HAlign_Left)
+				Panel->AddSlot(GridColumns_Notify, Row)
+				.HAlign(HAlign_Fill)
+				[
+					SNew(SBorder)
+					.BorderImage(FAppStyle::Get().GetBrush("WhiteBrush"))
+					.BorderBackgroundColor(EStyleColor::Background)
 					[
-						SNew(SEnumComboBox, StaticEnum<EEditorPerformanceNotificationOptions>())
-						.CurrentValue_Lambda([KPIPath]() -> int32
+						SNew(SBorder)
+						.BorderImage(FAppStyle::Get().GetBrush("NoBorder"))
+						.HAlign(HAlign_Center)
+						[
+							SNew(SCheckBox)
+							.IsChecked_Lambda([KPIPath]() -> ECheckBoxState
+								{
+									UEditorPerformanceSettings* EditorPerformanceSettings = GetMutableDefault<UEditorPerformanceSettings>();
+									if (EditorPerformanceSettings)
+									{
+										return EditorPerformanceSettings->NotifyList.Find(KPIPath) != INDEX_NONE ? ECheckBoxState::Checked : ECheckBoxState::Unchecked;
+									}
+									else
+									{
+										return ECheckBoxState::Unchecked;
+									}
+								})
+							.OnCheckStateChanged_Lambda([this, KPIPath](ECheckBoxState NewState)
 							{
 								UEditorPerformanceSettings* EditorPerformanceSettings = GetMutableDefault<UEditorPerformanceSettings>();
-								if (EditorPerformanceSettings)
-								{
-									return static_cast<int32>(EditorPerformanceSettings->NotificationList.Find(KPIPath) != INDEX_NONE ? EEditorPerformanceNotificationOptions::Notify : EEditorPerformanceNotificationOptions::Ignore);
-								}
-								else
-								{
-									return static_cast<int32>(EEditorPerformanceNotificationOptions::Ignore);
-								}
-							})
-						.OnEnumSelectionChanged_Lambda([this, KPIPath](int32 OptionValue, ESelectInfo::Type)
-						{
-							UEditorPerformanceSettings* EditorPerformanceSettings = GetMutableDefault<UEditorPerformanceSettings>();
 
-							if (EditorPerformanceSettings)
-							{
-								if ((EEditorPerformanceNotificationOptions)OptionValue == EEditorPerformanceNotificationOptions::Notify)
+								if (EditorPerformanceSettings)
 								{
-									if (EditorPerformanceSettings->NotificationList.Find(KPIPath) == INDEX_NONE)
+									if (NewState == ECheckBoxState::Checked)
 									{
-										// Add this KPI to the notification list
-										EditorPerformanceSettings->NotificationList.Emplace(KPIPath);
+										if (EditorPerformanceSettings->NotifyList.Find(KPIPath) == INDEX_NONE)
+										{
+											// Add this KPI to the notification list
+											EditorPerformanceSettings->NotifyList.Emplace(KPIPath);
+										}
+									}
+									else
+									{
+										// Remove this KPI to the notification ignore list
+										EditorPerformanceSettings->NotifyList.Remove(KPIPath);
 									}
-								}
-								else
-								{
-									// Remove this KPI to the notification ignore list
-									EditorPerformanceSettings->NotificationList.Remove(KPIPath);
-								}
-
-								EditorPerformanceSettings->PostEditChange();
-								EditorPerformanceSettings->SaveConfig();
-							}
 
-							UpdateGridPanels(0.0f, 0.0f);
-						})
-					];
+									EditorPerformanceSettings->PostEditChange();
+									EditorPerformanceSettings->SaveConfig();
+								}
+							})
+						]
+					]
+				];
 			}
+			else
+			{
+				Panel->AddSlot(GridColumns_Notify, Row)
+				.HAlign(HAlign_Fill)
+				[
+					SNew(SBorder)
+					.BorderImage(FAppStyle::Get().GetBrush("WhiteBrush"))
+					.BorderBackgroundColor(EStyleColor::Background)
+				];
+			}
+
+			Panel->AddSlot(GridColumns_Fill, Row)
+			.HAlign(HAlign_Fill)
+			[
+				CreateGridSlotWidget(FText(), DefaultMargin)
+			];
 
 			Row++;
 		}
@@ -748,4 +652,187 @@ TSharedRef<SWidget> SEditorPerformanceReportDialog::GetKPIGridPanel()
 	return Panel;
 }
 
+UToolMenu* SEditorPerformanceReportDialog::RegisterSettingsMenu()
+{
+	UToolMenus* ToolMenus = UToolMenus::Get();
+	if (!ToolMenus)
+	{
+		return nullptr;
+	}
+
+	UToolMenu* SettingsMenu = ToolMenus->RegisterMenu(SettingsMenuName);
+
+	FToolMenuSection& SettingsMenuSection = SettingsMenu->AddSection("Settings", LOCTEXT("Settings", "Settings"));
+
+	SettingsMenuSection.AddEntry(
+		FToolMenuEntry::InitMenuEntry("Notifications",
+			LOCTEXT("EnableNotificationsText", "Notifications"),
+			LOCTEXT("EnableNotificationsToolTip", "Enable All Notifications"),
+			FSlateIcon(),
+			FUIAction(
+				FExecuteAction::CreateLambda([this](){
+					UEditorPerformanceSettings* EditorPerformanceSettings = GetMutableDefault<UEditorPerformanceSettings>();
+
+					if (EditorPerformanceSettings)
+					{
+						EditorPerformanceSettings->bEnableSnapshots = !EditorPerformanceSettings->bEnableNotifications;
+						EditorPerformanceSettings->PostEditChange();
+						EditorPerformanceSettings->SaveConfig();
+
+						UpdateGridPanels(0.0f, 0.0f);
+					}
+				}),
+				FCanExecuteAction(),
+				FIsActionChecked::CreateLambda([this]()
+				{
+					const UEditorPerformanceSettings* EditorPerformanceSettings = GetDefault<UEditorPerformanceSettings>();
+					return EditorPerformanceSettings && EditorPerformanceSettings->bEnableNotifications;
+				})
+			),
+			EUserInterfaceActionType::Check));
+
+	SettingsMenuSection.AddEntry(
+		FToolMenuEntry::InitMenuEntry("Snapshots",
+			LOCTEXT("EnableSnapshotsText", "Snapshots"),
+			LOCTEXT("EnableSnapshotsToolTip", "Enable Automatic Capture of Unreal Insights Snapshots"),
+			FSlateIcon(),
+			FUIAction(
+				FExecuteAction::CreateLambda([this](){
+					UEditorPerformanceSettings* EditorPerformanceSettings = GetMutableDefault<UEditorPerformanceSettings>();
+
+					if (EditorPerformanceSettings)
+					{
+						EditorPerformanceSettings->bEnableSnapshots = !EditorPerformanceSettings->bEnableSnapshots;
+						EditorPerformanceSettings->PostEditChange();
+						EditorPerformanceSettings->SaveConfig();
+					}
+
+				}),
+				FCanExecuteAction(),
+				FIsActionChecked::CreateLambda([this]()
+				{
+					const UEditorPerformanceSettings* EditorPerformanceSettings = GetDefault<UEditorPerformanceSettings>();
+					return EditorPerformanceSettings && EditorPerformanceSettings->bEnableSnapshots;
+				})
+			),
+			EUserInterfaceActionType::Check));
+
+	SettingsMenuSection.AddEntry(
+		FToolMenuEntry::InitMenuEntry("Telemetry",
+			LOCTEXT("EnableTelemetryText", "Telemetry"),
+			LOCTEXT("EnableTelemetryToolTip", "Record Warning Telemetry Events To Analytics System"),
+			FSlateIcon(),
+			FUIAction(
+				FExecuteAction::CreateLambda([this](){
+					UEditorPerformanceSettings* EditorPerformanceSettings = GetMutableDefault<UEditorPerformanceSettings>();
+
+					if (EditorPerformanceSettings)
+					{
+						EditorPerformanceSettings->bEnableTelemetry = !EditorPerformanceSettings->bEnableTelemetry;
+						EditorPerformanceSettings->PostEditChange();
+						EditorPerformanceSettings->SaveConfig();
+					}
+				}),
+				FCanExecuteAction(),
+				FIsActionChecked::CreateLambda([this]()
+				{
+					const UEditorPerformanceSettings* EditorPerformanceSettings = GetDefault<UEditorPerformanceSettings>();
+					return EditorPerformanceSettings && EditorPerformanceSettings->bEnableTelemetry;
+				})
+			),
+			EUserInterfaceActionType::Check));
+
+	SettingsMenuSection.AddEntry(
+		FToolMenuEntry::InitMenuEntry("Throttling",
+			LOCTEXT("EnableBackgroundThrottlingText", "Throttling"),
+			LOCTEXT("EnableBackgroundThrottlingToolTip", "Enable CPU throttling when the Editor is in the background."),
+			FSlateIcon(),
+			FUIAction(
+				FExecuteAction::CreateLambda([this](){
+					UEditorPerformanceSettings* EditorPerformanceSettings = GetMutableDefault<UEditorPerformanceSettings>();
+
+					if (EditorPerformanceSettings)
+					{
+						EditorPerformanceSettings->bThrottleCPUWhenNotForeground = !EditorPerformanceSettings->bThrottleCPUWhenNotForeground;
+						EditorPerformanceSettings->PostEditChange();
+						EditorPerformanceSettings->SaveConfig();
+					}
+				}),
+				FCanExecuteAction(),
+				FIsActionChecked::CreateLambda([this]()
+				{
+					const UEditorPerformanceSettings* EditorPerformanceSettings = GetDefault<UEditorPerformanceSettings>();
+					return EditorPerformanceSettings && EditorPerformanceSettings->bThrottleCPUWhenNotForeground;
+				})
+			),
+			EUserInterfaceActionType::Check));
+
+	SettingsMenuSection.AddEntry(
+		FToolMenuEntry::InitMenuEntry("Diagnostics",
+			LOCTEXT("EnableShowFrameRateAndMemoryText", "Diagnostics"),
+			LOCTEXT("EnableShowFrameRateAndMemoryToolTip", "Show the Frame Rate, Memory and Stalls."),
+			FSlateIcon(),
+			FUIAction(
+				FExecuteAction::CreateLambda([this](){
+					UEditorPerformanceSettings* EditorPerformanceSettings = GetMutableDefault<UEditorPerformanceSettings>();
+
+					if (EditorPerformanceSettings)
+					{
+						EditorPerformanceSettings->bShowFrameRateAndMemory = !EditorPerformanceSettings->bShowFrameRateAndMemory;
+						EditorPerformanceSettings->PostEditChange();
+						EditorPerformanceSettings->SaveConfig();
+					}
+				}),
+				FCanExecuteAction(),
+				FIsActionChecked::CreateLambda([this]()
+				{
+					const UEditorPerformanceSettings* EditorPerformanceSettings = GetDefault<UEditorPerformanceSettings>();
+					return EditorPerformanceSettings && EditorPerformanceSettings->bShowFrameRateAndMemory;
+				})
+			),
+			EUserInterfaceActionType::Check));
+
+	SettingsMenuSection.AddSeparator(NAME_None);
+
+	SettingsMenuSection.AddEntry(
+		FToolMenuEntry::InitMenuEntry("AllSettings",
+			LOCTEXT("OpenSettingsText", "More Editor Performance Settings..."),
+			LOCTEXT("OpenSettingsToolTip", "Open the Editor Performance Settings Tab."),
+			FSlateIcon(),
+			FUIAction(
+				FExecuteAction::CreateLambda([]()
+				{
+					FModuleManager::LoadModuleChecked<ISettingsModule>("Settings").ShowViewer("Editor", "General", "EditorPerformanceSettings");
+				}),
+				FCanExecuteAction()
+			)));
+
+	return SettingsMenu;
+}
+
+TSharedRef<SWidget> SEditorPerformanceReportDialog::CreateSettingsMenuWidget()
+{
+	UToolMenus* ToolMenus = UToolMenus::Get();
+	if (!ToolMenus)
+	{
+		return SNullWidget::NullWidget;
+	}
+
+	UToolMenu* SettingsMenu = ToolMenus->FindMenu(SettingsMenuName);
+
+	if (!SettingsMenu)
+	{
+		SettingsMenu = RegisterSettingsMenu();
+	}
+
+	if (SettingsMenu)
+	{
+		return ToolMenus->GenerateWidget(SettingsMenu);
+	}
+	else
+	{
+		return SNullWidget::NullWidget;
+	}
+}
+
 #undef LOCTEXT_NAMESPACE
diff --git a/Engine/Plugins/Experimental/EditorPerformance/Source/EditorPerformance/Private/SEditorPerformanceDialogs.h b/Engine/Plugins/Experimental/EditorPerformance/Source/EditorPerformance/Private/SEditorPerformanceDialogs.h
--- a/Engine/Plugins/Experimental/EditorPerformance/Source/EditorPerformance/Private/SEditorPerformanceDialogs.h
+++ b/Engine/Plugins/Experimental/EditorPerformance/Source/EditorPerformance/Private/SEditorPerformanceDialogs.h
@@ -12,6 +12,7 @@
 #include "SEditorPerformanceDialogs.generated.h"
 
 class SWidget;
+class UToolMenu;
 
 UENUM()
 enum class EEditorPerformanceNotificationOptions : uint8
@@ -30,9 +31,13 @@ enum class EEditorPerformanceFilterOptions : uint8
 
 class SEditorPerformanceReportDialog : public SCompoundWidget
 {
+	static const TCHAR* SettingsMenuName;
+
 	SLATE_BEGIN_ARGS(SEditorPerformanceReportDialog) {}
 	SLATE_END_ARGS()
 
+	virtual ~SEditorPerformanceReportDialog();
+
 	void Construct(const FArguments& InArgs);
 
 private:
@@ -40,10 +45,11 @@ private:
 	EActiveTimerReturnType UpdateGridPanels(double InCurrentTime, float InDeltaTime);
 	
 	TSharedRef<SWidget> GetKPIGridPanel();
-	TSharedRef<SWidget> GetSettingsGridPanel();
 	TSharedRef<SWidget> GetHintGridPanel();
 
-	SVerticalBox::FSlot* SettingsGridSlot = nullptr;
+	UToolMenu* RegisterSettingsMenu();
+	TSharedRef<SWidget> CreateSettingsMenuWidget();
+
 	SVerticalBox::FSlot* KPIGridSlot = nullptr;
 	SVerticalBox::FSlot* HintGridSlot = nullptr;
 	uint32 CurrentHintIndex = 0;
diff --git a/Engine/Plugins/Experimental/EditorPerformance/Source/EditorPerformance/Private/SEditorPerformanceStatusBar.cpp b/Engine/Plugins/Experimental/EditorPerformance/Source/EditorPerformance/Private/SEditorPerformanceStatusBar.cpp
--- a/Engine/Plugins/Experimental/EditorPerformance/Source/EditorPerformance/Private/SEditorPerformanceStatusBar.cpp
+++ b/Engine/Plugins/Experimental/EditorPerformance/Source/EditorPerformance/Private/SEditorPerformanceStatusBar.cpp
@@ -113,31 +113,34 @@ EActiveTimerReturnType SEditorPerformanceStatusBarWidget::UpdateState(double InC
 
 		if (KPIValue.GetState() == FKPIValue::Bad)
 		{
-			// Currently exceeding the threshold	
-			EditorPerformanceState = EEditorPerformanceState::Warnings;
-
 			if (EditorPerformanceSettings)
 			{
-				if (EditorPerformanceSettings->NotificationList.Find(KPIValue.Path) != INDEX_NONE && AcknowledgedNotifications.Find(KPIValue.Path) == INDEX_NONE && CurrentNotificationName.IsNone())
+				if (EditorPerformanceSettings->NotifyList.Find(KPIValue.Path) != INDEX_NONE)
 				{
-					NotificationTitle = KPIValue.DisplayName;
+					// Turn into a warning state if the user wants to be notified	
+					EditorPerformanceState = EEditorPerformanceState::Warnings;
 
-					FKPIHint Hint;
-					if (EditorPerfModule.GetKPIRegistry().GetKPIHint(KPIValue.Id, Hint))
-					{
-						CurrentNotificationMessage = Hint.Message;
-					}
-					else
+					if (AcknowledgedNotifications.Find(KPIValue.Path) == INDEX_NONE && CurrentNotificationName.IsNone())
 					{
-						CurrentNotificationMessage = FText::FromString(*FString::Printf(TEXT("%s - %s was %s but should be %s %s"),
-							*KPIValue.Category.ToString(),
-							*KPIValue.Name.ToString(),
-							*FKPIValue::GetValueAsString(KPIValue.CurrentValue, KPIValue.DisplayType, KPIValue.CustomDisplayValueGetter),
-							*FKPIValue::GetComparisonAsPrettyString(KPIValue.Compare),
-							*FKPIValue::GetValueAsString(KPIValue.ThresholdValue.GetValue(), KPIValue.DisplayType, KPIValue.CustomDisplayValueGetter)));
-					}
+						NotificationTitle = KPIValue.DisplayName;
+
+						FKPIHint Hint;
+						if (EditorPerfModule.GetKPIRegistry().GetKPIHint(KPIValue.Id, Hint))
+						{
+							CurrentNotificationMessage = Hint.Message;
+						}
+						else
+						{
+							CurrentNotificationMessage = FText::FromString(*FString::Printf(TEXT("%s - %s was %s but should be %s %s"),
+								*KPIValue.Category.ToString(),
+								*KPIValue.Name.ToString(),
+								*FKPIValue::GetValueAsString(KPIValue.CurrentValue, KPIValue.DisplayType, KPIValue.CustomDisplayValueGetter),
+								*FKPIValue::GetComparisonAsPrettyString(KPIValue.Compare),
+								*FKPIValue::GetValueAsString(KPIValue.ThresholdValue.GetValue(), KPIValue.DisplayType, KPIValue.CustomDisplayValueGetter)));
+						}
 
-					CurrentNotificationName = KPIValue.Path;
+						CurrentNotificationName = KPIValue.Path;
+					}
 				}
 			}
 
diff --git a/Engine/Source/Editor/EditorStyle/Private/StarshipStyle.cpp b/Engine/Source/Editor/EditorStyle/Private/StarshipStyle.cpp
--- a/Engine/Source/Editor/EditorStyle/Private/StarshipStyle.cpp
+++ b/Engine/Source/Editor/EditorStyle/Private/StarshipStyle.cpp
@@ -1158,8 +1158,8 @@ void FStarshipEditorStyle::FStyle::SetupGeneralStyles()
 	// Editor Performance Data
 	{
 		Set("EditorPerformance.Settings", new IMAGE_BRUSH_SVG("Starship/DerivedData/DD_Cache_Settings", Icon16x16));
-		Set("EditorPerformance.Notification.Warning", new IMAGE_BRUSH_SVG("Starship/Common/Stall", Icon16x16, EStyleColor::Warning));
-		Set("EditorPerformance.Notification.Good", new IMAGE_BRUSH_SVG("Starship/Common/Stall", Icon16x16));
+		Set("EditorPerformance.Notification.Warning", new IMAGE_BRUSH_SVG("Starship/DerivedData/DD_Cache_Statistics", Icon16x16, EStyleColor::Warning));
+		Set("EditorPerformance.Notification.Good", new IMAGE_BRUSH_SVG("Starship/DerivedData/DD_Cache_Statistics", Icon16x16));
 		Set("EditorPerformance.Report.Panel", new IMAGE_BRUSH_SVG("Starship/DerivedData/DD_Cache_Statistics", Icon16x16));
 		Set("EditorPerformance.Report.Warning", new CORE_IMAGE_BRUSH_SVG("Starship/Common/alert-triangle", Icon16x16, EStyleColor::Warning));
 		Set("EditorPerformance.Stall", new IMAGE_BRUSH_SVG("Starship/Common/Stall", Icon16x16));
diff --git a/Engine/Source/Editor/UnrealEd/Classes/Editor/EditorPerformanceSettings.h b/Engine/Source/Editor/UnrealEd/Classes/Editor/EditorPerformanceSettings.h
--- a/Engine/Source/Editor/UnrealEd/Classes/Editor/EditorPerformanceSettings.h
+++ b/Engine/Source/Editor/UnrealEd/Classes/Editor/EditorPerformanceSettings.h
@@ -42,8 +42,12 @@ class UEditorPerformanceSettings : public UDeveloperSettings
 	bool bEnableNotifications = false;
 
 	/** Specific measurements to notify */
-	UPROPERTY(EditAnywhere, config, Category = EditorPerformanceTool, meta = (DisplayName = "Notification List", ConfigRestartRequired = false, EditCondition = "bEnableEditorPeformanceTool"))
-	TArray<FName> NotificationList;
+	UPROPERTY(EditAnywhere, config, Category = EditorPerformanceTool, meta = (ConfigRestartRequired = false, EditCondition = "bEnableEditorPeformanceTool"))
+	TArray<FName> NotifyList;
+
+	/** Specific measurements to notify */
+	UPROPERTY(config, meta = (DisplayName = "Notification List", ConfigRestartRequired = false))
+	TArray<FName> NotificationList_DEPRECATED;
 
 	/** Whether to enable Telemetry */
 	UPROPERTY(EditAnywhere, config, Category = EditorPerformanceTool, meta = (DisplayName = "Enable Telemetry", ConfigRestartRequired = false, EditCondition = "bEnableEditorPeformanceTool"))
