Editor Performance Dialog:
- Fixed menu toggling the wrong setting

diff --git a/Engine/Plugins/Experimental/EditorPerformance/Source/EditorPerformance/Private/SEditorPerformanceDialogs.cpp b/Engine/Plugins/Experimental/EditorPerformance/Source/EditorPerformance/Private/SEditorPerformanceDialogs.cpp
--- a/Engine/Plugins/Experimental/EditorPerformance/Source/EditorPerformance/Private/SEditorPerformanceDialogs.cpp
+++ b/Engine/Plugins/Experimental/EditorPerformance/Source/EditorPerformance/Private/SEditorPerformanceDialogs.cpp
@@ -675,7 +675,7 @@ UToolMenu* SEditorPerformanceReportDialog::RegisterSettingsMenu()
 
 					if (EditorPerformanceSettings)
 					{
-						EditorPerformanceSettings->bEnableSnapshots = !EditorPerformanceSettings->bEnableNotifications;
+						EditorPerformanceSettings->bEnableNotifications = !EditorPerformanceSettings->bEnableNotifications;
 						EditorPerformanceSettings->PostEditChange();
 						EditorPerformanceSettings->SaveConfig();
 
