Editor Diagnostics status bar styling

diff --git a/Engine/Plugins/Experimental/EditorPerformance/Source/EditorPerformance/Private/SEditorPerformanceStatusBar.cpp b/Engine/Plugins/Experimental/EditorPerformance/Source/EditorPerformance/Private/SEditorPerformanceStatusBar.cpp
--- a/Engine/Plugins/Experimental/EditorPerformance/Source/EditorPerformance/Private/SEditorPerformanceStatusBar.cpp
+++ b/Engine/Plugins/Experimental/EditorPerformance/Source/EditorPerformance/Private/SEditorPerformanceStatusBar.cpp
@@ -67,21 +67,41 @@ void SEditorPerformanceStatusBarWidget::Construct(const FArguments& InArgs)
 			.ButtonStyle(&FAppStyle::Get().GetWidgetStyle<FButtonStyle>("StatusBar.StatusBarButton"))
 			.Content()
 			[
-				SNew(SHorizontalBox)
-				+ SHorizontalBox::Slot()
-				.AutoWidth()
-				.VAlign(VAlign_Center)
-				.Padding(0.f)
+				SNew(SBox)
+				.HAlign(HAlign_Fill)
+				.Padding(6.f, 0.f)
 				[
-					SNew(SOverlay)
-					+ SOverlay::Slot()
-					.HAlign(HAlign_Center)
-					.VAlign(VAlign_Top)
+					SNew(SHorizontalBox)
+					+SHorizontalBox::Slot()
+					.AutoWidth()
+					.VAlign(VAlign_Center)
+					.Padding(0.f, 0.f, 3.f, 0.f)
 					[
-						SNew(SImage)
-						.ColorAndOpacity_Lambda([this] { return GetStatusIconColor(); })
-						.Image_Lambda([this] { return GetStatusIcon(); })
-						.ToolTipText_Lambda([this] { return GetStatusToolTipText(); })
+						SNew(SOverlay)
+						+SOverlay::Slot()
+						.HAlign(HAlign_Center)
+						.VAlign(VAlign_Center)
+						[
+							SNew(SImage)
+							.ColorAndOpacity(FSlateColor::UseForeground())
+							.Image_Lambda([this] { return GetStatusIcon(); })
+							.ToolTipText_Lambda([this] { return GetStatusToolTipText(); })
+						]
+						+ SOverlay::Slot()
+						.HAlign(HAlign_Center)
+						.VAlign(VAlign_Center)
+						[
+							SNew(SImage)
+							.ColorAndOpacity(FSlateColor::UseForeground())
+							.Image_Lambda([this] { return GetStatusBadgeIcon(); })
+						]
+					]
+					+SHorizontalBox::Slot()
+					.FillWidth(1.f)
+					.VAlign(VAlign_Center)
+					[
+						SNew(STextBlock)
+						.Text(LOCTEXT("Diagnostics", "Diagnostics"))
 					]
 				]
 			]
@@ -94,7 +114,7 @@ void SEditorPerformanceStatusBarWidget::Construct(const FArguments& InArgs)
 EActiveTimerReturnType SEditorPerformanceStatusBarWidget::UpdateState(double InCurrentTime, float InDeltaTime)
 {
 	EditorPerformanceState = EEditorPerformanceState::Good;
-	EditorPerformanceStateMessage = LOCTEXT("EditorPerfMesssageGood", "Good");
+	EditorPerformanceStateMessage = LOCTEXT("ToolTipMessageGood", "Nothing to report.");
 
 	FEditorPerformanceModule& EditorPerfModule = FModuleManager::LoadModuleChecked<FEditorPerformanceModule>("EditorPerformance");
 
@@ -156,14 +176,7 @@ EActiveTimerReturnType SEditorPerformanceStatusBarWidget::UpdateState(double InC
 
 	if (WarningCount > 0)
 	{
-		if (WarningCount == 1)
-		{
-			EditorPerformanceStateMessage = FText::FromString(*FString::Printf(TEXT("There Is 1 Warning.\n\nView Performance Report For Details.")));
-		}
-		else
-		{
-			EditorPerformanceStateMessage = FText::FromString(*FString::Printf(TEXT("There Are %d Warnings.\n\nView Performance Report For Details."), WarningCount));
-		}
+		EditorPerformanceStateMessage = LOCTEXT("ToolTipWarning", "Warning. View report for details.");
 	}
 
 	if (EditorPerformanceSettings && EditorPerformanceSettings->bEnableNotifications)
@@ -230,34 +243,36 @@ const FSlateBrush* SEditorPerformanceStatusBarWidget::GetStatusIcon() const
 		default:
 		case EEditorPerformanceState::Good:
 		{
-			return FAppStyle::Get().GetBrush("EditorPerformance.Notification.Good");
+			return FAppStyle::Get().GetBrush("EditorDiagnostics.StatusBar.Icon");
 		}
 
 		case EEditorPerformanceState::Warnings:
 		{
-			return FAppStyle::Get().GetBrush("EditorPerformance.Notification.Warning");
+			return FAppStyle::Get().GetBrush("EditorDiagnostics.StatusBar.BadgeBG");
 		}
 	}
 }
 
-const FSlateColor SEditorPerformanceStatusBarWidget::GetStatusIconColor() const
-{
-	return FSlateColor::UseForeground();
-}
-
-FText SEditorPerformanceStatusBarWidget::GetStatusToolTipText() const
+const FSlateBrush* SEditorPerformanceStatusBarWidget::GetStatusBadgeIcon() const
 {
-	return EditorPerformanceStateMessage;
-}
+	switch (EditorPerformanceState)
+	{
+		default:
+		case EEditorPerformanceState::Good:
+		{
+			return FAppStyle::Get().GetBrush("NoBrush");
+		}
 
-FText SEditorPerformanceStatusBarWidget::GetTitleText() const
-{
-	return LOCTEXT("EditorPerformanceToolBarName", "Performance");
+		case EEditorPerformanceState::Warnings:
+		{
+			return FAppStyle::Get().GetBrush("EditorDiagnostics.StatusBar.WarningBadge");
+		}
+	}
 }
 
-FText SEditorPerformanceStatusBarWidget::GetTitleToolTipText() const
+FText SEditorPerformanceStatusBarWidget::GetStatusToolTipText() const
 {
-	return GetTitleText();
+	return FText::Format(LOCTEXT("ToolTipFormat", "Editor Diagnostics: {0}"), EditorPerformanceStateMessage);
 }
 
 #undef LOCTEXT_NAMESPACE
diff --git a/Engine/Plugins/Experimental/EditorPerformance/Source/EditorPerformance/Private/SEditorPerformanceStatusBar.h b/Engine/Plugins/Experimental/EditorPerformance/Source/EditorPerformance/Private/SEditorPerformanceStatusBar.h
--- a/Engine/Plugins/Experimental/EditorPerformance/Source/EditorPerformance/Private/SEditorPerformanceStatusBar.h
+++ b/Engine/Plugins/Experimental/EditorPerformance/Source/EditorPerformance/Private/SEditorPerformanceStatusBar.h
@@ -33,10 +33,8 @@ class SEditorPerformanceStatusBarWidget : public SCompoundWidget
 private:
 	static FReply ViewPerformanceReport_Clicked();
 
-	FText							GetTitleToolTipText() const;
-	FText							GetTitleText() const;
 	const FSlateBrush*				GetStatusIcon() const;
-	const FSlateColor				GetStatusIconColor() const;
+	const FSlateBrush*				GetStatusBadgeIcon() const;
 	FText							GetStatusToolTipText() const;
 	EActiveTimerReturnType			UpdateState(double InCurrentTime, float InDeltaTime);
 	
diff --git a/Engine/Source/Editor/EditorStyle/Private/StarshipStyle.cpp b/Engine/Source/Editor/EditorStyle/Private/StarshipStyle.cpp
--- a/Engine/Source/Editor/EditorStyle/Private/StarshipStyle.cpp
+++ b/Engine/Source/Editor/EditorStyle/Private/StarshipStyle.cpp
@@ -1155,14 +1155,14 @@ void FStarshipEditorStyle::FStyle::SetupGeneralStyles()
 		Set("DerivedData.RemoteCache.Unavailable", new IMAGE_BRUSH_SVG("Starship/DerivedData/DD_RemoteCache_Unavailable", Icon16x16));
 	}
 
-	// Editor Performance Data
+	// Editor Diagnostics
 	{
-		Set("EditorPerformance.Settings", new IMAGE_BRUSH_SVG("Starship/DerivedData/DD_Cache_Settings", Icon16x16));
-		Set("EditorPerformance.Notification.Warning", new IMAGE_BRUSH_SVG("Starship/DerivedData/DD_Cache_Statistics", Icon16x16, EStyleColor::Warning));
-		Set("EditorPerformance.Notification.Good", new IMAGE_BRUSH_SVG("Starship/DerivedData/DD_Cache_Statistics", Icon16x16));
-		Set("EditorPerformance.Report.Panel", new IMAGE_BRUSH_SVG("Starship/DerivedData/DD_Cache_Statistics", Icon16x16));
+		Set("EditorPerformance.Report.Panel", new IMAGE_BRUSH_SVG("Starship/Common/Diagnostics_16", Icon16x16));
 		Set("EditorPerformance.Report.Warning", new CORE_IMAGE_BRUSH_SVG("Starship/Common/alert-triangle", Icon16x16, EStyleColor::Warning));
-		Set("EditorPerformance.Stall", new IMAGE_BRUSH_SVG("Starship/Common/Stall", Icon16x16));
+
+		Set("EditorDiagnostics.StatusBar.Icon", new IMAGE_BRUSH_SVG("Starship/Common/Diagnostics_16", Icon16x16));
+		Set("EditorDiagnostics.StatusBar.BadgeBG", new IMAGE_BRUSH_SVG("Starship/Common/DiagnosticsBadgeBG_16", Icon16x16));
+		Set("EditorDiagnostics.StatusBar.WarningBadge", new CORE_IMAGE_BRUSH_SVG("Starship/Common/alert-triangle-badge", Icon16x16, EStyleColor::Warning));
 
 		Set("EditorDiagnostics.SetShowActivityMonitor", new IMAGE_BRUSH_SVG("Starship/Common/Scalability_16", Icon20x20));
 		Set("EditorDiagnostics.SetShowStalls", new IMAGE_BRUSH_SVG("Starship/Common/Stall", Icon20x20));
