Editor Diagnostics
- Auto dismiss the notification if KPIs go back to a good state
- Added a "Don't show this again" option that will remove the KPi from the notify list

diff --git a/Engine/Plugins/Experimental/EditorPerformance/Source/EditorPerformance/Private/SEditorPerformanceStatusBar.cpp b/Engine/Plugins/Experimental/EditorPerformance/Source/EditorPerformance/Private/SEditorPerformanceStatusBar.cpp
--- a/Engine/Plugins/Experimental/EditorPerformance/Source/EditorPerformance/Private/SEditorPerformanceStatusBar.cpp
+++ b/Engine/Plugins/Experimental/EditorPerformance/Source/EditorPerformance/Private/SEditorPerformanceStatusBar.cpp
@@ -121,7 +121,7 @@ EActiveTimerReturnType SEditorPerformanceStatusBarWidget::UpdateState(double InC
 	const UEditorPerformanceSettings* EditorPerformanceSettings = GetDefault<UEditorPerformanceSettings>();
 
 	WarningCount = 0;
-
+	
 	EditorPerfModule.UpdateKPIs(InDeltaTime);
 
 	FText NotificationTitle;
@@ -171,6 +171,11 @@ EActiveTimerReturnType SEditorPerformanceStatusBarWidget::UpdateState(double InC
 			// No longer exceeding threshold, so no need to acknowledge the last time it was raised to the user
 			// There may be subsequent times that this same KPI is exceeded this session so we may want to alert the user again
 			AcknowledgedNotifications.Remove(KPIValue.Path);
+
+			if (CurrentNotificationName == KPIValue.Path)
+			{
+				CurrentNotificationName = FName();
+			}
 		}
 	}
 
@@ -179,58 +184,93 @@ EActiveTimerReturnType SEditorPerformanceStatusBarWidget::UpdateState(double InC
 		EditorPerformanceStateMessage = LOCTEXT("ToolTipWarning", "Warning. View report for details.");
 	}
 
-	if (EditorPerformanceSettings && EditorPerformanceSettings->bEnableNotifications)
-	{
-		if (CurrentNotificationName.IsNone() == false)
-		{	
-			if (NotificationItem.IsValid() == false || NotificationItem->GetCompletionState() == SNotificationItem::CS_None)
-			{
-				FNotificationInfo Info(FText::Format(LOCTEXT("NotificationTitle", "{0} Warning"), NotificationTitle));
+	if (CurrentNotificationName.IsNone() == false && EditorPerformanceSettings && EditorPerformanceSettings->bEnableNotifications)
+	{	
+		if (NotificationItem.IsValid() == false || NotificationItem->GetCompletionState() == SNotificationItem::CS_None)
+		{
+			FNotificationInfo Info(FText::Format(LOCTEXT("NotificationTitle", "{0} Warning"), NotificationTitle));
 
-				Info.SubText = CurrentNotificationMessage;
-				Info.bUseSuccessFailIcons = true;
-				Info.bFireAndForget = false;
-				Info.bUseThrobber = true;
-				Info.FadeOutDuration = 1.0f;
-				Info.ExpireDuration = 0.0f;
+			Info.SubText = CurrentNotificationMessage;
+			Info.bUseSuccessFailIcons = true;
+			Info.bFireAndForget = false;
+			Info.bUseThrobber = true;
+			Info.FadeOutDuration = 1.0f;
+			Info.ExpireDuration = 0.0f;
 
-				// No existing notification or the existing one has finished
-				TPromise<TWeakPtr<SNotificationItem>> AcknowledgeNotificationPromise;
+			// No existing notification or the existing one has finished
+			TPromise<TWeakPtr<SNotificationItem>> AcknowledgeNotificationPromise;
 
-				Info.ButtonDetails.Add(FNotificationButtonInfo(LOCTEXT("AcknowledgeNotificationButton", "Dismiss"), FText(), FSimpleDelegate::CreateLambda([NotificationFuture = AcknowledgeNotificationPromise.GetFuture().Share(),this]()
+			Info.ButtonDetails.Add(FNotificationButtonInfo(LOCTEXT("AcknowledgeNotificationButton", "Dismiss"), FText(), FSimpleDelegate::CreateLambda([NotificationFuture = AcknowledgeNotificationPromise.GetFuture().Share(),this]()
+				{
+					// User has acknowledged this warning
+					TWeakPtr<SNotificationItem> NotificationPtr = NotificationFuture.Get();
+					if (TSharedPtr<SNotificationItem> Notification = NotificationPtr.Pin())
 					{
-						// User has acknowledged this warning
-						TWeakPtr<SNotificationItem> NotificationPtr = NotificationFuture.Get();
-						if (TSharedPtr<SNotificationItem> Notification = NotificationPtr.Pin())
+						Notification->SetCompletionState(SNotificationItem::CS_None);
+						Notification->ExpireAndFadeout();
+					}
+
+					AcknowledgedNotifications.Add(CurrentNotificationName);
+					CurrentNotificationName = FName();
+
+				}), SNotificationItem::ECompletionState::CS_Fail));
+
+			// Add a "Don't show this again" option
+			Info.CheckBoxState = TAttribute<ECheckBoxState>::CreateLambda([CurrentNotificationName=this->CurrentNotificationName]()
+				{
+					if (CurrentNotificationName.IsNone() == false)
+					{
+						if (const UEditorPerformanceSettings* EditorPerformanceSettings = GetDefault<UEditorPerformanceSettings>())
 						{
-							Notification->SetCompletionState(SNotificationItem::CS_None);
-							Notification->ExpireAndFadeout();
+							return EditorPerformanceSettings->NotifyList.Contains(CurrentNotificationName) ? ECheckBoxState::Unchecked : ECheckBoxState::Checked;
 						}
+					}
 
-						AcknowledgedNotifications.Add(CurrentNotificationName);
-						CurrentNotificationName = FName();
+					return ECheckBoxState::Unchecked;
+				});
+			Info.CheckBoxStateChanged = FOnCheckStateChanged::CreateLambda([CurrentNotificationName=this->CurrentNotificationName](ECheckBoxState NewState)
+				{
+					if (CurrentNotificationName.IsNone())
+					{
+						return;
+					}
+
+					if (UEditorPerformanceSettings* EditorPerformanceSettings = GetMutableDefault<UEditorPerformanceSettings>())
+					{
+						switch (NewState)
+						{
+							case ECheckBoxState::Checked:
+								EditorPerformanceSettings->NotifyList.Remove(CurrentNotificationName);
+								break;
+							case ECheckBoxState::Unchecked:
+								EditorPerformanceSettings->NotifyList.AddUnique(CurrentNotificationName);
+								break;
+							default:
+								break;
+						}
 
-					}), SNotificationItem::ECompletionState::CS_Fail));
+						EditorPerformanceSettings->PostEditChange();
+						EditorPerformanceSettings->SaveConfig();
+					}
+				});
+			Info.CheckBoxText = LOCTEXT("DontShowThisAgainCheckBoxMessage", "Don't show this again");
 
-				// Create the notification item
-				NotificationItem = FSlateNotificationManager::Get().AddNotification(Info);
 
-				if (NotificationItem.IsValid())
-				{
-					AcknowledgeNotificationPromise.SetValue(NotificationItem);
-					NotificationItem->SetCompletionState(SNotificationItem::CS_Fail);
-				}
+			// Create the notification item
+			NotificationItem = FSlateNotificationManager::Get().AddNotification(Info);
+
+			if (NotificationItem.IsValid())
+			{
+				AcknowledgeNotificationPromise.SetValue(NotificationItem);
+				NotificationItem->SetCompletionState(SNotificationItem::CS_Fail);
 			}
 		}
 	}
-	else
+	// No longer any warnings so kill any existing notifications
+	else if (NotificationItem.IsValid())
 	{
-		// No longer any warnings so kill any existing notifications
-		if (NotificationItem.IsValid())
-		{
-			NotificationItem->SetCompletionState(SNotificationItem::CS_None);
-			NotificationItem->ExpireAndFadeout();
-		}
+		NotificationItem->SetCompletionState(SNotificationItem::CS_None);
+		NotificationItem->ExpireAndFadeout();
 	}
 
 	return EActiveTimerReturnType::Continue;
