Editor Diagnostics: Made the status and notification more reactive
- Use an event instead of a 5s timer to update the status bar and notification
- Remove the call to UpdateKPIs coming from the status bar since the Editor Perf Module already updates those

diff --git a/Engine/Plugins/Experimental/EditorPerformance/Source/EditorPerformance/Private/EditorPerformanceModule.cpp b/Engine/Plugins/Experimental/EditorPerformance/Source/EditorPerformance/Private/EditorPerformanceModule.cpp
--- a/Engine/Plugins/Experimental/EditorPerformance/Source/EditorPerformance/Private/EditorPerformanceModule.cpp
+++ b/Engine/Plugins/Experimental/EditorPerformance/Source/EditorPerformance/Private/EditorPerformanceModule.cpp
@@ -608,8 +608,12 @@ void FEditorPerformanceModule::UpdateKPIs(float InDeltaTime)
 		KPIRegistry.SetKPIValue(PIEStallRateKPI, StallRate);
 	}
 
+	DebugForceKPIIntoBadStateIfNeeded();
+
 	static TArray<FGuid> RecordedKPIEvent;
 	
+	const bool bGoodStateBefore = RecordedKPIEvent.IsEmpty();
+
 	// Check for KPIs that have exceeded their value
 	for (FKPIValues::TConstIterator It(GetKPIRegistry().GetKPIValues()); It; ++It)
 	{
@@ -648,7 +652,12 @@ void FEditorPerformanceModule::UpdateKPIs(float InDeltaTime)
 		}
 	}
 
-	DebugForceKPIIntoBadStateIfNeeded();
+	const bool bGoodStateAfter = RecordedKPIEvent.IsEmpty();
+
+	if (bGoodStateBefore != bGoodStateAfter)
+	{
+		OnPerformanceStateChanged.Broadcast();
+	}
 }
 
 void FEditorPerformanceModule::DebugForceKPIIntoBadStateIfNeeded()
diff --git a/Engine/Plugins/Experimental/EditorPerformance/Source/EditorPerformance/Private/SEditorPerformanceStatusBar.cpp b/Engine/Plugins/Experimental/EditorPerformance/Source/EditorPerformance/Private/SEditorPerformanceStatusBar.cpp
--- a/Engine/Plugins/Experimental/EditorPerformance/Source/EditorPerformance/Private/SEditorPerformanceStatusBar.cpp
+++ b/Engine/Plugins/Experimental/EditorPerformance/Source/EditorPerformance/Private/SEditorPerformanceStatusBar.cpp
@@ -108,10 +108,11 @@ void SEditorPerformanceStatusBarWidget::Construct(const FArguments& InArgs)
 			.OnClicked(FOnClicked::CreateStatic(&SEditorPerformanceStatusBarWidget::ViewPerformanceReport_Clicked))
 		];
 
-	RegisterActiveTimer(5.0f, FWidgetActiveTimerDelegate::CreateSP(this, &SEditorPerformanceStatusBarWidget::UpdateState));
+	FEditorPerformanceModule& EditorPerfModule = FModuleManager::LoadModuleChecked<FEditorPerformanceModule>("EditorPerformance");
+	EditorPerfModule.GetOnPerformanceStateChanged().AddSP(this, &SEditorPerformanceStatusBarWidget::UpdateState);
 }
 
-EActiveTimerReturnType SEditorPerformanceStatusBarWidget::UpdateState(double InCurrentTime, float InDeltaTime)
+void SEditorPerformanceStatusBarWidget::UpdateState()
 {
 	EditorPerformanceState = EEditorPerformanceState::Good;
 	EditorPerformanceStateMessage = LOCTEXT("ToolTipMessageGood", "Nothing to report.");
@@ -121,8 +122,6 @@ EActiveTimerReturnType SEditorPerformanceStatusBarWidget::UpdateState(double InC
 	const UEditorPerformanceSettings* EditorPerformanceSettings = GetDefault<UEditorPerformanceSettings>();
 
 	WarningCount = 0;
-	
-	EditorPerfModule.UpdateKPIs(InDeltaTime);
 
 	FText NotificationTitle;
 
@@ -272,8 +271,6 @@ EActiveTimerReturnType SEditorPerformanceStatusBarWidget::UpdateState(double InC
 		NotificationItem->SetCompletionState(SNotificationItem::CS_None);
 		NotificationItem->ExpireAndFadeout();
 	}
-
-	return EActiveTimerReturnType::Continue;
 }
 
 const FSlateBrush* SEditorPerformanceStatusBarWidget::GetStatusIcon() const
diff --git a/Engine/Plugins/Experimental/EditorPerformance/Source/EditorPerformance/Private/SEditorPerformanceStatusBar.h b/Engine/Plugins/Experimental/EditorPerformance/Source/EditorPerformance/Private/SEditorPerformanceStatusBar.h
--- a/Engine/Plugins/Experimental/EditorPerformance/Source/EditorPerformance/Private/SEditorPerformanceStatusBar.h
+++ b/Engine/Plugins/Experimental/EditorPerformance/Source/EditorPerformance/Private/SEditorPerformanceStatusBar.h
@@ -36,7 +36,7 @@ private:
 	const FSlateBrush*				GetStatusIcon() const;
 	const FSlateBrush*				GetStatusBadgeIcon() const;
 	FText							GetStatusToolTipText() const;
-	EActiveTimerReturnType			UpdateState(double InCurrentTime, float InDeltaTime);
+	void							UpdateState();
 	
 	TSharedPtr<SNotificationItem>	NotificationItem;
 	EEditorPerformanceState			EditorPerformanceState= EEditorPerformanceState::Good;
diff --git a/Engine/Plugins/Experimental/EditorPerformance/Source/EditorPerformance/Public/EditorPerformanceModule.h b/Engine/Plugins/Experimental/EditorPerformance/Source/EditorPerformance/Public/EditorPerformanceModule.h
--- a/Engine/Plugins/Experimental/EditorPerformance/Source/EditorPerformance/Public/EditorPerformanceModule.h
+++ b/Engine/Plugins/Experimental/EditorPerformance/Source/EditorPerformance/Public/EditorPerformanceModule.h
@@ -3,6 +3,7 @@
 #pragma once
 
 #include "CoreMinimal.h"
+#include "Delegates/Delegate.h"
 #include "Engine/EngineTypes.h"
 #include "Modules/ModuleInterface.h"
 #include "Modules/ModuleManager.h"
@@ -58,18 +59,21 @@ public:
 	UE_API bool IsHotLocalCacheCase() const;
 	UE_API FEditorPerformanceModule::EEditorState GetEditorState() const;
 
+	DECLARE_EVENT(FEditorPerformanceModule, FPerformanceStateChangedEvent);
+	UE_API FPerformanceStateChangedEvent& GetOnPerformanceStateChanged() { return OnPerformanceStateChanged; }
+
 private:
-	UE_API TSharedRef<SDockTab> CreatePerformanceReportTab(const FSpawnTabArgs& Args);	
+	TSharedRef<SDockTab> CreatePerformanceReportTab(const FSpawnTabArgs& Args);	
 	TWeakPtr<SDockTab> PerformanceReportTab;
 
-	UE_API void InitializeEditor();
-	UE_API void TerminateEditor();
+	void InitializeEditor();
+	void TerminateEditor();
 
-	UE_API void InitializeKPIs();
-	UE_API void TerminateKPIs();
+	void InitializeKPIs();
+	void TerminateKPIs();
 
-	UE_API void HeartBeatCallback();
-	UE_API void HitchSamplerCallback();
+	void HeartBeatCallback();
+	void HitchSamplerCallback();
 
 	void DebugForceKPIIntoBadStateIfNeeded();
 
@@ -78,6 +82,9 @@ private:
 	FKPIRegistry					KPIRegistry;
 	TMap<FString,FKPIProfile>		KPIProfiles;
 	FString							KPIProfileName=TEXT("Default");
+
+	FPerformanceStateChangedEvent	OnPerformanceStateChanged;
+
 	FDateTime						LoadMapStartTime;
 	FDateTime						PIEStartTime;
 	FDateTime						PIEEndTime;
