Slate Invalidation - Show Slate Invalidation status in Diagnostic Window

diff --git a/Engine/Plugins/Experimental/EditorPerformance/Source/EditorPerformance/Private/EditorPerformanceModule.cpp b/Engine/Plugins/Experimental/EditorPerformance/Source/EditorPerformance/Private/EditorPerformanceModule.cpp
--- a/Engine/Plugins/Experimental/EditorPerformance/Source/EditorPerformance/Private/EditorPerformanceModule.cpp
+++ b/Engine/Plugins/Experimental/EditorPerformance/Source/EditorPerformance/Private/EditorPerformanceModule.cpp
@@ -102,11 +102,15 @@ const FText AvailableMemoryKPIText = LOCTEXT("AvailableMemory", "Available Memor
 const FName MemoryPressureKPIName = TEXT("Memory Pressure");
 const FText MemoryPressureKPIText = LOCTEXT("MemoryPressure", "Memory Pressure");
 
+const FName DynamicInvalidationKPIName = TEXT("Dynamic Invalidation");
+const FText DynamicInvalidationKPIText = LOCTEXT("DynamicInvalidation", "Dynamic Invalidation");
+
 // Thresholds
 TOptional<float> NoSetLimit;
 float LocalDDCEffciencyKPILimit = 85.f;
 float VirtualAssetEfficiencyKPILimit = 95.f;
 float MemoryPressureKPILimit = (float)FGenericPlatformMemoryStats::EMemoryPressureStatus::Warning;
+float DynamicInvalidationKPILimit = 1.0f;
  
 static const FName EditorDiagnosticsTabName = FName(TEXT("EditorDiagnosticsTab"));
 
@@ -241,7 +245,12 @@ void FEditorPerformanceModule::InitializeKPIs()
 	MemoryPressureKPIValue.CustomDisplayValueGetter.BindRaw(this, &FEditorPerformanceModule::GetMemoryPressureDisplayValue);
 
 	MemoryPressureKPI			= KPIRegistry.DeclareKPIValue(MoveTemp(MemoryPressureKPIValue));
-	
+
+	FKPIValue DynamicInvalidationKPIValue = FKPIValue(EditorCategoryName, EditorCategoryText, DynamicInvalidationKPIName,
+		DynamicInvalidationKPIText, 0.0, DynamicInvalidationKPILimit, FKPIValue::GreaterThanOrEqual, FKPIValue::Custom);
+	DynamicInvalidationKPIValue.CustomDisplayValueGetter.BindRaw(this, &FEditorPerformanceModule::GetDynamicInvalidationDisplayValue);
+	DynamicInvalidationKPI		= KPIRegistry.DeclareKPIValue(MoveTemp(DynamicInvalidationKPIValue));
+
 	// Declare the KPI Hints
 	KPIRegistry.DeclareKPIHint(EditorBootKPI, LOCTEXT("EditorBootHintMessage", "The Editor boot time is slow.\nCheck you have enabled a Game Feature Plugin profile for your project and that the expected local cache efficiency is met.\nIf you are booting the Editor in the background then disable the Use Less CPU in Background option in the settings."), LOCTEXT("EditorBootHintURL","https://docs.unrealengine.com/5.0/en-US/"));
 	KPIRegistry.DeclareKPIHint(TotalTimeToEditorKPI, LOCTEXT("EditorStartupHintMessage", "The Editor start-up time is slow.\nCheck you have enabled a Game Feature Plugin profile for your project and that the expected local cache efficiency is met.\nIf you are booting the Editor in the background then disable the Use Less CPU in Background option in the settings."), LOCTEXT("EditorBootHintURL","https://docs.unrealengine.com/5.0/en-US/"));
@@ -254,6 +263,7 @@ void FEditorPerformanceModule::InitializeKPIs()
 	KPIRegistry.DeclareKPIHint(AvailableMemoryKPI, LOCTEXT("LowAvailableMemoryHintMessage", "Your hardware is running low on memory.\nTry closing applications that are no longer needed to recover available memory."), LOCTEXT("LowAvailableMemoryHintURL", "https://docs.unrealengine.com/5.0/en-US/"));
 	KPIRegistry.DeclareKPIHint(MemoryPressureKPI, LOCTEXT("MemoryPressureHintMessage", "Your system is running low on memory.\nTry closing applications that are no longer needed to recover available memory."), LOCTEXT("LowAvailableMemoryHintURL", "https://docs.unrealengine.com/5.0/en-US/"));
 
+	KPIRegistry.DeclareKPIHint(DynamicInvalidationKPI, LOCTEXT("DynamicInvalidationHintMessage", "Some Dynamic Invalidation Panels contain widgets that do not support invalidation."), LOCTEXT("DynamicInvalidationHintURL", "https://docs.unrealengine.com/5.0/en-US/"));
 
 	// Load the KPI profiles
 	KPIRegistry.LoadKPIProfiles(TEXT("EditorPerformance.Profile"), GEditorIni);
@@ -608,6 +618,37 @@ void FEditorPerformanceModule::UpdateKPIs(float InDeltaTime)
 		KPIRegistry.SetKPIValue(PIEStallRateKPI, StallRate);
 	}
 
+#if WITH_SLATE_DEBUGGING && UE_SLATE_WITH_DYNAMIC_INVALIDATION
+	{
+		IConsoleVariable* DynamicInvalidationOptionsVar = IConsoleManager::Get().FindConsoleVariable(TEXT("Slate.DynamicInvalidation.Options"));
+		ensure(DynamicInvalidationOptionsVar);
+		const int32 DynamicInvalidationOptions = DynamicInvalidationOptionsVar ? DynamicInvalidationOptionsVar->GetInt() : 0;
+
+		const int NoPanels = FSlateDebugging::GetNoInvalidationPanels();
+		const int NoDynPanels = FSlateDebugging::GetNoDynamicInvalidationPanels();
+		const int NoDynPanelsSupportingInvalidation = FSlateDebugging::GetNoDynamicInvalidationPanelsSupportingInvalidation();
+
+		if (DynamicInvalidationOptions == 0) // Dynamic Invalidation Disabled
+		{
+			KPIRegistry.InvalidateKPIValue(DynamicInvalidationKPI);
+		}
+		else if (DynamicInvalidationOptions == 1 && NoPanels > 0) // Dynamic Invalidation Panels acting like Invalidation Panels
+		{
+			KPIRegistry.SetKPIValue(DynamicInvalidationKPI, 1.0); // all of them are invalidating
+		}
+		else if (DynamicInvalidationOptions > 1 && NoDynPanels > 0) // Dynamic Invalidation enabled or forced
+		{
+			KPIRegistry.SetKPIValue(DynamicInvalidationKPI, (float)NoDynPanelsSupportingInvalidation / (float)NoDynPanels);
+		}
+		else
+		{
+			KPIRegistry.InvalidateKPIValue(DynamicInvalidationKPI);
+		}
+	}
+#else
+	KPIRegistry.InvalidateKPIValue(DynamicInvalidationKPI);
+#endif
+
 	DebugForceKPIIntoBadStateIfNeeded();
 
 	static TArray<FGuid> RecordedKPIEvent;
@@ -690,6 +731,38 @@ FText FEditorPerformanceModule::GetMemoryPressureDisplayValue(const float InValu
 	}
 }
 
+FText FEditorPerformanceModule::GetDynamicInvalidationDisplayValue(const float InValue) const
+{
+#if UE_SLATE_WITH_DYNAMIC_INVALIDATION
+	IConsoleVariable* DynamicInvalidationOptionsVar = IConsoleManager::Get().FindConsoleVariable(TEXT("Slate.DynamicInvalidation.Options"));
+	ensure(DynamicInvalidationOptionsVar);
+	const int32 DynamicInvalidationOptions = DynamicInvalidationOptionsVar ? DynamicInvalidationOptionsVar->GetInt() : 0;
+	
+	if (DynamicInvalidationOptions == 0) // Dynamic Invalidation Disabled
+	{
+		return LOCTEXT("DynamicInvalidationDisabled", "Disabled");
+	}
+
+	const FText Percent = InValue < 0 ? FText::FromString("...") : FText::AsPercent(InValue);
+	if (DynamicInvalidationOptions == 1) // Dynamic Invalidation Panels acting like regular invalidation panels
+	{
+		return FText::Format(LOCTEXT("DynamicInvalidationForcedStatic", "{0} (Forced Static)"), Percent);
+	}
+	else if (DynamicInvalidationOptions == 2) // Dynamic Invalidation enabled
+	{
+		return Percent;
+	}
+	else if (DynamicInvalidationOptions == 3) // Invalidation Panels acting like regular Dynamic Invalidation Panels
+	{
+		return FText::Format(LOCTEXT("DynamicInvalidationForcedDynamic", "{0} (Forced Dynamic)"), Percent);
+	}
+
+	return LOCTEXT("DynamicInvalidationUnknown", "Unknown");
+#else
+	return LOCTEXT("DynamicInvalidationDisabled", "Disabled");
+#endif
+}
+
 bool FEditorPerformanceModule::IsHotLocalCacheCase() const
 {
 	FKPIValue KPIValue;
diff --git a/Engine/Plugins/Experimental/EditorPerformance/Source/EditorPerformance/Public/EditorPerformanceModule.h b/Engine/Plugins/Experimental/EditorPerformance/Source/EditorPerformance/Public/EditorPerformanceModule.h
--- a/Engine/Plugins/Experimental/EditorPerformance/Source/EditorPerformance/Public/EditorPerformanceModule.h
+++ b/Engine/Plugins/Experimental/EditorPerformance/Source/EditorPerformance/Public/EditorPerformanceModule.h
@@ -78,6 +78,7 @@ private:
 	void DebugForceKPIIntoBadStateIfNeeded();
 
 	FText GetMemoryPressureDisplayValue(const float) const;
+	FText GetDynamicInvalidationDisplayValue(const float) const;
 
 	FKPIRegistry					KPIRegistry;
 	TMap<FString,FKPIProfile>		KPIProfiles;
@@ -131,6 +132,7 @@ private:
 	FGuid							VirtualAssetEfficiencyKPI;
 	FGuid							AvailableMemoryKPI;
 	FGuid							MemoryPressureKPI;
+	FGuid							DynamicInvalidationKPI;
 };
 
 
diff --git a/Engine/Source/Developer/SlateReflector/Private/Widgets/SWidgetReflectorTreeWidgetItem.cpp b/Engine/Source/Developer/SlateReflector/Private/Widgets/SWidgetReflectorTreeWidgetItem.cpp
--- a/Engine/Source/Developer/SlateReflector/Private/Widgets/SWidgetReflectorTreeWidgetItem.cpp
+++ b/Engine/Source/Developer/SlateReflector/Private/Widgets/SWidgetReflectorTreeWidgetItem.cpp
@@ -79,7 +79,7 @@ TSharedRef<SWidget> SReflectorTreeWidgetItem::GenerateWidgetForColumn(const FNam
 			InvalidationPanel = InvalidationPanel->GetParentWidget();
 		}
 
-		if (InvalidationPanel && StaticCastSharedPtr<SInvalidationPanel>(InvalidationPanel)->GetUseDynamicInvalidation()) // if we found a Dynamic Invalidation Panel
+		if (InvalidationPanel && StaticCastSharedPtr<SInvalidationPanel>(InvalidationPanel)->CanUseDynamicInvalidation()) // if we found an Invalidation Panel using Dynamic Invalidation
 		{
 			const bool CanCache = StaticCastSharedPtr<SInvalidationPanel>(InvalidationPanel)->GetCanCache();
 			// we always show a color block for the Dynamic Invalidation Panel, but for its children we only show the color if it can cache
diff --git a/Engine/Source/Runtime/Slate/Private/Widgets/SInvalidationPanel.cpp b/Engine/Source/Runtime/Slate/Private/Widgets/SInvalidationPanel.cpp
--- a/Engine/Source/Runtime/Slate/Private/Widgets/SInvalidationPanel.cpp
+++ b/Engine/Source/Runtime/Slate/Private/Widgets/SInvalidationPanel.cpp
@@ -53,14 +53,24 @@ namespace DynamicInvalidationCVars
 	FConsoleVariableDelegate::CreateStatic(ConsoleVariableEnableDynamicInvalidationChanged)
 	);
 
+	/** Is Dynamic Invalidation and caching disabled for all Dynamic Invalidation Panels? */
+	bool IsDynamicInvalidationDisabled() { return DynamicInvalidationOptions == 0; }
+	/** Is Dynamic Invalidation reverted to regular Invalidation? */
+	bool IsDynamicInvalidationForcedStatic() { return DynamicInvalidationOptions == 1; }
+	/** Is Dynamic Invalidation enabled? */
+	bool IsDynamicInvalidationEnabled() { return DynamicInvalidationOptions == 2; }
+	/** Is Dynamic Invalidation forced on all invalidation panels? */
+	bool IsDynamicInvalidationForced() { return DynamicInvalidationOptions == 3; }
+	
 	void ConsoleVariableEnableDynamicInvalidationChanged(IConsoleVariable*)
 	{
-		DynamicInvalidationOptions = FMath::Clamp(DynamicInvalidationOptions, 0, 4);
+		DynamicInvalidationOptions = FMath::Clamp(DynamicInvalidationOptions, 0, 3);
 #if UE_SLATE_TRACE_ENABLED
-		const FString Description = DynamicInvalidationOptions == 0 ? TEXT("Disabled") :
-			DynamicInvalidationOptions == 1 ? TEXT("Disabled: forced Static") :
-			DynamicInvalidationOptions == 2 ? TEXT("Enabled") :
-			DynamicInvalidationOptions == 3 ? TEXT("Forced") : TEXT("Unknown");
+		const FString Description =
+			IsDynamicInvalidationDisabled()		? TEXT("Disabled") :
+			IsDynamicInvalidationForcedStatic()	? TEXT("Disabled: forced Static") :
+			IsDynamicInvalidationEnabled()		? TEXT("Enabled") :
+			IsDynamicInvalidationForced()		? TEXT("Forced") : TEXT("Unknown");
 
 		UE_TRACE_SLATE_BOOKMARK(TEXT("Dynamic Invalidation changed [%d]: %s"), DynamicInvalidationOptions, *Description)
 #endif // UE_SLATE_TRACE_ENABLED
@@ -161,15 +171,9 @@ void SInvalidationPanel::EnableInvalidationPanels(bool bEnable)
 bool SInvalidationPanel::GetCanCache() const
 {
 #if UE_SLATE_WITH_DYNAMIC_INVALIDATION
-	const bool bForcedDynamicInvalidation = DynamicInvalidationCVars::DynamicInvalidationOptions == 3;
-	const bool bDynamicInvalidationEnabled = bForcedDynamicInvalidation || (bUseDynamicInvalidation && DynamicInvalidationCVars::DynamicInvalidationOptions == 2);
-	if (!bDynamicInvalidationEnabled)
+	if (DynamicInvalidationCVars::IsDynamicInvalidationDisabled() & bUseDynamicInvalidation)
 	{
-		const bool bRevertToRegularInvalidation = !bUseDynamicInvalidation || DynamicInvalidationCVars::DynamicInvalidationOptions == 1;
-		if (!bRevertToRegularInvalidation)
-		{
-			return false;
-		}
+		return false;
 	}
 #endif
 	return bCanCache && !GSlateEnableGlobalInvalidation && bInvalidationPanelsEnabled;
@@ -254,6 +258,12 @@ void SInvalidationPanel::SetUseDynamicInvalidation(bool InUseDynamicInvalidation
 	const bool bUseCachedValue = false;
 	SupportsInvalidationRecursive(bUseCachedValue);
 }
+
+bool SInvalidationPanel::CanUseDynamicInvalidation() const
+{
+	const bool bDynamicInvalidationEnabled = bUseDynamicInvalidation && DynamicInvalidationCVars::IsDynamicInvalidationEnabled();
+	return DynamicInvalidationCVars::IsDynamicInvalidationForced() | bDynamicInvalidationEnabled;
+}
 #endif
 
 FChildren* SInvalidationPanel::GetChildren()
@@ -278,10 +288,8 @@ FChildren* SInvalidationPanel::Debug_GetChildrenForReflector()
 #if UE_SLATE_WITH_DYNAMIC_INVALIDATION
 bool SInvalidationPanel::SupportsInvalidationRecursive(bool bUseCachedValue) const
 {
-	const bool bForcedDynamicInvalidation = DynamicInvalidationCVars::DynamicInvalidationOptions == 3;
-	const bool bDynamicInvalidationEnabled = bForcedDynamicInvalidation || (bUseDynamicInvalidation && DynamicInvalidationCVars::DynamicInvalidationOptions == 2);
 	// For an invalidation panel without bUseDynamicInvalidation, we force the caching of all children, so we always support invalidation
-	return !bDynamicInvalidationEnabled || SWidget::SupportsInvalidationRecursive(bUseCachedValue);
+	return !CanUseDynamicInvalidation() || SWidget::SupportsInvalidationRecursive(bUseCachedValue);
 }
 #endif
 
@@ -296,7 +304,12 @@ int32 SInvalidationPanel::OnPaint( const FPaintArgs& Args, const FGeometry& Allo
 	SInvalidationPanel* MutableThis = const_cast<SInvalidationPanel*>(this);
 
 #if UE_SLATE_WITH_DYNAMIC_INVALIDATION
-	const bool bCanCacheThisFrame = GetCanCache() && SupportsInvalidationRecursive();
+	const bool bSupportsInvalidationRecursive = SupportsInvalidationRecursive();
+	const bool bCanCacheValue = GetCanCache();
+	const bool bCanCacheThisFrame = bCanCacheValue & bSupportsInvalidationRecursive;
+#if WITH_SLATE_DEBUGGING
+	FSlateDebugging::IncrementInvalidationPanelCount(bCanCacheValue, CanUseDynamicInvalidation(), bSupportsInvalidationRecursive);
+#endif
 #else
 	const bool bCanCacheThisFrame = GetCanCache();
 #endif
diff --git a/Engine/Source/Runtime/Slate/Public/Widgets/SInvalidationPanel.h b/Engine/Source/Runtime/Slate/Public/Widgets/SInvalidationPanel.h
--- a/Engine/Source/Runtime/Slate/Public/Widgets/SInvalidationPanel.h
+++ b/Engine/Source/Runtime/Slate/Public/Widgets/SInvalidationPanel.h
@@ -63,13 +63,23 @@ public:
 
 #if UE_SLATE_WITH_DYNAMIC_INVALIDATION
 	/**
-	 * Returns true when the invalidation panel uses Dynamic Invalidation.
+	 * Returns true when the invalidation panel should use Dynamic Invalidation.
 	 * Dynamic Invalidation means that the widget's content will only be cached if all descendant widgets support invalidation.
+	 * Dynamic Invalidation can be globally disabled or forced. Use `CanUseDynamicInvalidation` to check if this invalidation panel will actually
+	 * use Dynamic Invalidation
 	**/
-	SLATE_API virtual bool GetUseDynamicInvalidation() const;
-
+	SLATE_API bool GetUseDynamicInvalidation() const;
 	SLATE_API void SetUseDynamicInvalidation(bool InUseDynamicInvalidation);
 
+	/**
+	 * Checks if this invalidation panel will actually use Dynamic Invalidation, depending on the CVar `Slate.DynamicInvalidation.Options`
+	 * If Dynamic Invalidation is disabled, this will return false.
+	 * If this invalidation panel is not dynamic, or if Dynamic Invalidation is forced Static, this will return false.
+	 * If Dynamic Invalidation is enabled, this will return the same value as `GetUseDynamicInvalidation`.
+	 * If Dynamic Invalidation is forced, this will return true.
+	 */
+	SLATE_API bool CanUseDynamicInvalidation() const;
+
 	SLATE_API virtual bool SupportsInvalidationRecursive(bool bUseCachedValue = true) const override;
 #endif
 
diff --git a/Engine/Source/Runtime/SlateCore/Private/Debugging/SlateDebugging.cpp b/Engine/Source/Runtime/SlateCore/Private/Debugging/SlateDebugging.cpp
--- a/Engine/Source/Runtime/SlateCore/Private/Debugging/SlateDebugging.cpp
+++ b/Engine/Source/Runtime/SlateCore/Private/Debugging/SlateDebugging.cpp
@@ -702,6 +702,35 @@ void FSlateDebugging::ExportWidgetList(FStringView Filename)
 #endif
 }
 
+uint32 FSlateDebugging::NoInvalidationPanels = 0;
+uint32 FSlateDebugging::NoDynamicInvalidationPanels = 0;
+uint32 FSlateDebugging::NoDynamicInvalidationPanelsSupportingInvalidation = 0;
+
+void FSlateDebugging::IncrementInvalidationPanelCount(bool bCanCache, bool bCanUseDynamic, bool bSupportsInvalidation)
+{
+	if (!bCanCache)
+	{
+		return;
+	}
+
+	++NoInvalidationPanels;
+	if (bCanUseDynamic)
+	{
+		++NoDynamicInvalidationPanels;
+		if (bSupportsInvalidation)
+		{
+			++NoDynamicInvalidationPanelsSupportingInvalidation;
+		}
+	}
+}
+
+void FSlateDebugging::ResetInvalidationPanelCount()
+{
+	NoInvalidationPanels = 0;
+	NoDynamicInvalidationPanels = 0;
+	NoDynamicInvalidationPanelsSupportingInvalidation = 0;
+}
+
 #undef LOCTEXT_NAMESPACE
 
 #endif
diff --git a/Engine/Source/Runtime/SlateCore/Public/Debugging/SlateDebugging.h b/Engine/Source/Runtime/SlateCore/Public/Debugging/SlateDebugging.h
--- a/Engine/Source/Runtime/SlateCore/Public/Debugging/SlateDebugging.h
+++ b/Engine/Source/Runtime/SlateCore/Public/Debugging/SlateDebugging.h
@@ -436,6 +436,13 @@ public:
 	static SLATECORE_API const TArray<const SWidget*>& GetAllWidgets();
 	static SLATECORE_API void ExportWidgetList(FStringView Filename);
 
+public:
+	static SLATECORE_API void IncrementInvalidationPanelCount(bool bCanCache, bool bCanUseDynamic, bool bSupportsInvalidation);
+	static SLATECORE_API void ResetInvalidationPanelCount();
+	static uint32 GetNoInvalidationPanels() { return NoInvalidationPanels; }
+	static uint32 GetNoDynamicInvalidationPanels() { return NoDynamicInvalidationPanels; }
+	static uint32 GetNoDynamicInvalidationPanelsSupportingInvalidation() { return NoDynamicInvalidationPanelsSupportingInvalidation; }
+
 private:
 
 	// This class is only for namespace use
@@ -450,6 +457,10 @@ private:
 	static FLastCursorQuery LastCursorQuery;
 
 	static TArray<IWidgetInputRoutingEvent*> RoutingEvents;
+
+	static SLATECORE_API uint32 NoInvalidationPanels;
+	static SLATECORE_API uint32 NoDynamicInvalidationPanels;
+	static SLATECORE_API uint32 NoDynamicInvalidationPanelsSupportingInvalidation;
 };
 
 #endif
