Editor KPIs: unhardcoded the Hint URLs

diff --git a/Engine/Plugins/Experimental/EditorPerformance/Source/EditorPerformance/Private/EditorPerformanceModule.cpp b/Engine/Plugins/Experimental/EditorPerformance/Source/EditorPerformance/Private/EditorPerformanceModule.cpp
--- a/Engine/Plugins/Experimental/EditorPerformance/Source/EditorPerformance/Private/EditorPerformanceModule.cpp
+++ b/Engine/Plugins/Experimental/EditorPerformance/Source/EditorPerformance/Private/EditorPerformanceModule.cpp
@@ -252,18 +252,18 @@ void FEditorPerformanceModule::InitializeKPIs()
 	DynamicInvalidationKPI		= KPIRegistry.DeclareKPIValue(MoveTemp(DynamicInvalidationKPIValue));
 
 	// Declare the KPI Hints
-	KPIRegistry.DeclareKPIHint(EditorBootKPI, LOCTEXT("EditorBootHintMessage", "The Editor boot time is slow.\nCheck you have enabled a Game Feature Plugin profile for your project and that the expected local cache efficiency is met.\nIf you are booting the Editor in the background then disable the Use Less CPU in Background option in the settings."), LOCTEXT("EditorBootHintURL","https://docs.unrealengine.com/5.0/en-US/"));
-	KPIRegistry.DeclareKPIHint(TotalTimeToEditorKPI, LOCTEXT("EditorStartupHintMessage", "The Editor start-up time is slow.\nCheck you have enabled a Game Feature Plugin profile for your project and that the expected local cache efficiency is met.\nIf you are booting the Editor in the background then disable the Use Less CPU in Background option in the settings."), LOCTEXT("EditorBootHintURL","https://docs.unrealengine.com/5.0/en-US/"));
-	KPIRegistry.DeclareKPIHint(EditorPluginCountKPI, LOCTEXT("EditorPluginHintMessage", "The Editor is loading more plugins than expected and this will affect Editor start-up performance.\nCheck you have enabled a Game Feature Plugin profile for your project."), LOCTEXT("EditorPluginHintURL", "https://docs.unrealengine.com/5.0/en-US/"));
+	KPIRegistry.DeclareKPIHint(EditorBootKPI, LOCTEXT("EditorBootHintMessage", "The Editor boot time is slow.\nCheck you have enabled a Game Feature Plugin profile for your project and that the expected local cache efficiency is met.\nIf you are booting the Editor in the background then disable the Use Less CPU in Background option in the settings."));
+	KPIRegistry.DeclareKPIHint(TotalTimeToEditorKPI, LOCTEXT("EditorStartupHintMessage", "The Editor start-up time is slow.\nCheck you have enabled a Game Feature Plugin profile for your project and that the expected local cache efficiency is met.\nIf you are booting the Editor in the background then disable the Use Less CPU in Background option in the settings."));
+	KPIRegistry.DeclareKPIHint(EditorPluginCountKPI, LOCTEXT("EditorPluginHintMessage", "The Editor is loading more plugins than expected and this will affect Editor start-up performance.\nCheck you have enabled a Game Feature Plugin profile for your project."));
 
-	KPIRegistry.DeclareKPIHint(PIETransitionKPI, LOCTEXT("PIETransitionHintMessage", "The Editor transition to PIE is slow.\nCheck that the expected local cache efficiency is met.\nIf you are transitioning to PIE with the Editor in the background then disable the Use Less CPU in Background option in the settings."), LOCTEXT("PIETransitionHintURL", "https://docs.unrealengine.com/5.0/en-US/"));
+	KPIRegistry.DeclareKPIHint(PIETransitionKPI, LOCTEXT("PIETransitionHintMessage", "The Editor transition to PIE is slow.\nCheck that the expected local cache efficiency is met.\nIf you are transitioning to PIE with the Editor in the background then disable the Use Less CPU in Background option in the settings."));
 	
-	KPIRegistry.DeclareKPIHint(LocalDDCEfficiencyKPI, LOCTEXT("LocalCacheEfficiencyHintMessage", "The Editor will not perform well if the local cache efficiency has not yet met the expected value.\nIf this is the first time you have booted the Editor after a sync then this is to be expected."), LOCTEXT("EditorCacheHintURL","https://docs.unrealengine.com/5.3/en-US/derived-data-cache/"));
+	KPIRegistry.DeclareKPIHint(LocalDDCEfficiencyKPI, LOCTEXT("LocalCacheEfficiencyHintMessage", "The Editor will not perform well if the local cache efficiency has not yet met the expected value.\nIf this is the first time you have booted the Editor after a sync then this is to be expected."), TEXT("using-derived-data-cache-in-unreal-engine"));
 	
-	KPIRegistry.DeclareKPIHint(AvailableMemoryKPI, LOCTEXT("LowAvailableMemoryHintMessage", "Your hardware is running low on memory.\nTry closing applications that are no longer needed to recover available memory."), LOCTEXT("LowAvailableMemoryHintURL", "https://docs.unrealengine.com/5.0/en-US/"));
-	KPIRegistry.DeclareKPIHint(MemoryPressureKPI, LOCTEXT("MemoryPressureHintMessage", "Your system is running low on memory.\nTry closing applications that are no longer needed to recover available memory."), LOCTEXT("LowAvailableMemoryHintURL", "https://docs.unrealengine.com/5.0/en-US/"));
+	KPIRegistry.DeclareKPIHint(AvailableMemoryKPI, LOCTEXT("LowAvailableMemoryHintMessage", "Your hardware is running low on memory.\nTry closing applications that are no longer needed to recover available memory."));
+	KPIRegistry.DeclareKPIHint(MemoryPressureKPI, LOCTEXT("MemoryPressureHintMessage", "Your system is running low on memory.\nTry closing applications that are no longer needed to recover available memory."));
 
-	KPIRegistry.DeclareKPIHint(DynamicInvalidationKPI, LOCTEXT("DynamicInvalidationHintMessage", "Some Dynamic Invalidation Panels contain widgets that do not support invalidation."), LOCTEXT("DynamicInvalidationHintURL", "https://docs.unrealengine.com/5.0/en-US/"));
+	KPIRegistry.DeclareKPIHint(DynamicInvalidationKPI, LOCTEXT("DynamicInvalidationHintMessage", "Some Dynamic Invalidation Panels contain widgets that do not support invalidation."));
 
 	// Load the KPI profiles
 	KPIRegistry.LoadKPIProfiles(TEXT("EditorPerformance.Profile"), GEditorIni);
diff --git a/Engine/Plugins/Experimental/EditorPerformance/Source/EditorPerformance/Private/KPIValue.cpp b/Engine/Plugins/Experimental/EditorPerformance/Source/EditorPerformance/Private/KPIValue.cpp
--- a/Engine/Plugins/Experimental/EditorPerformance/Source/EditorPerformance/Private/KPIValue.cpp
+++ b/Engine/Plugins/Experimental/EditorPerformance/Source/EditorPerformance/Private/KPIValue.cpp
@@ -257,7 +257,7 @@ FGuid FKPIRegistry::DeclareKPIValue( const FKPIValue& Value )
 	return FGuid();
 }
 
-bool FKPIRegistry::DeclareKPIHint(FGuid Id, const FText& HintMessage, const FText& HintURL)
+bool FKPIRegistry::DeclareKPIHint(FGuid Id, const FText& HintMessage, const FString& DocumentationLink)
 {
 	if (Values.Find(Id) != nullptr)
 	{
@@ -265,7 +265,7 @@ bool FKPIRegistry::DeclareKPIHint(FGuid Id, const FText& HintMessage, const FTex
 
 		NewKPIHint.Id = Id;
 		NewKPIHint.Message = HintMessage;
-		NewKPIHint.URL = HintURL;
+		NewKPIHint.DocumentationLink = DocumentationLink;
 
 		FKPIHint* KPIHint = Hints.Find(Id);
 
diff --git a/Engine/Plugins/Experimental/EditorPerformance/Source/EditorPerformance/Private/SEditorPerformanceDialogs.cpp b/Engine/Plugins/Experimental/EditorPerformance/Source/EditorPerformance/Private/SEditorPerformanceDialogs.cpp
--- a/Engine/Plugins/Experimental/EditorPerformance/Source/EditorPerformance/Private/SEditorPerformanceDialogs.cpp
+++ b/Engine/Plugins/Experimental/EditorPerformance/Source/EditorPerformance/Private/SEditorPerformanceDialogs.cpp
@@ -21,6 +21,7 @@
 #include "Widgets/Input/SHyperlink.h"
 #include "EditorPerformanceModule.h"
 #include "Modules/ModuleManager.h"
+#include "IDocumentation.h"
 #include "ISettingsModule.h"
 #include "Editor/EditorPerformanceSettings.h"
 #include "SEnumCombo.h"
@@ -197,7 +198,7 @@ TSharedRef<SWidget> SEditorPerformanceReportDialog::GetHintGridPanel()
 
 			Row++;
 
-			if (!KPIHint.URL.IsEmpty())
+			if (!KPIHint.DocumentationLink.IsEmpty())
 			{
 				Panel->AddSlot(0, Row)
 					.HAlign(HAlign_Left)
@@ -205,8 +206,8 @@ TSharedRef<SWidget> SEditorPerformanceReportDialog::GetHintGridPanel()
 					[
 						SNew(SHyperlink)
 							.Text(LOCTEXT("HintLinkName", "Further Help & Documentation"))
-							.ToolTipText_Lambda([=]() { return FText::FromString(*KPIHint.URL.ToString()); })
-							.OnNavigate_Lambda([=]() { FPlatformProcess::LaunchURL(*KPIHint.URL.ToString(), nullptr, nullptr); })
+							.ToolTipText(FText::FromString(*KPIHint.DocumentationLink))
+							.OnNavigate_Lambda([KPIHint]() { IDocumentation::Get()->Open(KPIHint.DocumentationLink); })
 					];
 
 				Row++;
diff --git a/Engine/Plugins/Experimental/EditorPerformance/Source/EditorPerformance/Public/KPIValue.h b/Engine/Plugins/Experimental/EditorPerformance/Source/EditorPerformance/Public/KPIValue.h
--- a/Engine/Plugins/Experimental/EditorPerformance/Source/EditorPerformance/Public/KPIValue.h
+++ b/Engine/Plugins/Experimental/EditorPerformance/Source/EditorPerformance/Public/KPIValue.h
@@ -107,7 +107,7 @@ class FKPIHint
 public:
 	FGuid			Id;
 	FText			Message;
-	FText			URL;
+	FString			DocumentationLink;
 };
 
 typedef TMap<FGuid, FKPIHint> FKPIHints;
@@ -118,7 +118,7 @@ public:
 
 	FGuid							DeclareKPIValue(const FName Category, const FText DisplayCategory, const FName Name, const FText DisplayName, float InitialValue, TOptional<float> ThresholdValue, FKPIValue::ECompare Compare, FKPIValue::EDisplayType Type);
 	FGuid							DeclareKPIValue(const FKPIValue& Value);
-	bool							DeclareKPIHint(FGuid Id, const FText& HintMessage, const FText& HintURL);
+	bool							DeclareKPIHint(FGuid Id, const FText& HintMessage, const FString& DocumentationLink = FString());
 
 	bool							SetKPIValue(FGuid Id, float CurrentValue);
 	bool							SetKPIThreshold(FGuid Id, float ThresholdValue);
