Fix warning for missing editor extension point
- A "Compile" section doesn't appear to exist. The only section is "Insights".

#Synced-CL 48490803

diff --git a/Engine/Plugins/Experimental/EditorPerformance/Source/EditorPerformance/Private/EditorPerformanceModule.cpp b/Engine/Plugins/Experimental/EditorPerformance/Source/EditorPerformance/Private/EditorPerformanceModule.cpp
--- a/Engine/Plugins/Experimental/EditorPerformance/Source/EditorPerformance/Private/EditorPerformanceModule.cpp
+++ b/Engine/Plugins/Experimental/EditorPerformance/Source/EditorPerformance/Private/EditorPerformanceModule.cpp
@@ -167,7 +167,7 @@ void FEditorPerformanceModule::InitializeEditor()
 	if (Menu != nullptr)
 	{
 		// Add the Editor Perf ToolBar 
-		FToolMenuSection& EditorPerfSection = Menu->AddSection("EditorPerf", FText::GetEmpty(), FToolMenuInsert("Compile", EToolMenuInsertType::Before));
+		FToolMenuSection& EditorPerfSection = Menu->AddSection("EditorPerf", FText::GetEmpty(), FToolMenuInsert("Insights", EToolMenuInsertType::Before));
 		EditorPerfSection.AddEntry(FToolMenuEntry::InitWidget("EditorPerformanceStatusBar", CreateStatusBarWidget(), FText::GetEmpty(), true, false));
 	}
 	
