Making calls to `FindConsoleVariable(TEXT("Slate.DynamicInvalidation.Options"))` static

#Synced-CL 48711348

diff --git a/Engine/Plugins/Experimental/EditorPerformance/Source/EditorPerformance/Private/EditorPerformanceModule.cpp b/Engine/Plugins/Experimental/EditorPerformance/Source/EditorPerformance/Private/EditorPerformanceModule.cpp
--- a/Engine/Plugins/Experimental/EditorPerformance/Source/EditorPerformance/Private/EditorPerformanceModule.cpp
+++ b/Engine/Plugins/Experimental/EditorPerformance/Source/EditorPerformance/Private/EditorPerformanceModule.cpp
@@ -620,7 +620,7 @@ void FEditorPerformanceModule::UpdateKPIs(float InDeltaTime)
 
 #if WITH_SLATE_DEBUGGING && UE_SLATE_WITH_DYNAMIC_INVALIDATION
 	{
-		IConsoleVariable* DynamicInvalidationOptionsVar = IConsoleManager::Get().FindConsoleVariable(TEXT("Slate.DynamicInvalidation.Options"));
+		static IConsoleVariable* DynamicInvalidationOptionsVar = IConsoleManager::Get().FindConsoleVariable(TEXT("Slate.DynamicInvalidation.Options"));
 		ensure(DynamicInvalidationOptionsVar);
 		const int32 DynamicInvalidationOptions = DynamicInvalidationOptionsVar ? DynamicInvalidationOptionsVar->GetInt() : 0;
 
@@ -734,7 +734,7 @@ FText FEditorPerformanceModule::GetMemoryPressureDisplayValue(const float InValu
 FText FEditorPerformanceModule::GetDynamicInvalidationDisplayValue(const float InValue) const
 {
 #if UE_SLATE_WITH_DYNAMIC_INVALIDATION
-	IConsoleVariable* DynamicInvalidationOptionsVar = IConsoleManager::Get().FindConsoleVariable(TEXT("Slate.DynamicInvalidation.Options"));
+	static IConsoleVariable* DynamicInvalidationOptionsVar = IConsoleManager::Get().FindConsoleVariable(TEXT("Slate.DynamicInvalidation.Options"));
 	ensure(DynamicInvalidationOptionsVar);
 	const int32 DynamicInvalidationOptions = DynamicInvalidationOptionsVar ? DynamicInvalidationOptionsVar->GetInt() : 0;
 	
