Editor Performance KPIs
- Added severity to KPIs. Major and Critical severity KPIs will always show in the status bar when failing.
- Marked dynamic invalidation as Major and memory pressure as Critical

#Synced-CL 49280314

diff --git a/Engine/Plugins/Experimental/EditorPerformance/Source/EditorPerformance/Private/EditorPerformanceModule.cpp b/Engine/Plugins/Experimental/EditorPerformance/Source/EditorPerformance/Private/EditorPerformanceModule.cpp
--- a/Engine/Plugins/Experimental/EditorPerformance/Source/EditorPerformance/Private/EditorPerformanceModule.cpp
+++ b/Engine/Plugins/Experimental/EditorPerformance/Source/EditorPerformance/Private/EditorPerformanceModule.cpp
@@ -239,7 +239,7 @@ void FEditorPerformanceModule::InitializeKPIs()
 	VirtualAssetEfficiencyKPI	= KPIRegistry.DeclareKPIValue(CacheCategoryName, CacheCategoryText, VirtualAssetEfficiencyKPIName, VirtualAssetEfficiencyKPIText, 100.0, VirtualAssetEfficiencyKPILimit, FKPIValue::GreaterThan, FKPIValue::Percent);
 	AvailableMemoryKPI			= KPIRegistry.DeclareKPIValue(MemoryCategoryName, MemoryCategoryText, AvailableMemoryKPIName, AvailableMemoryKPIText, 128.0, NoSetLimit, FKPIValue::GreaterThanOrEqual, FKPIValue::GigaBytes);
 
-	FKPIValue MemoryPressureKPIValue = FKPIValue(MemoryCategoryName, MemoryCategoryText, MemoryPressureKPIName, MemoryPressureKPIText,
+	FKPIValue MemoryPressureKPIValue = FKPIValue(MemoryCategoryName, MemoryCategoryText, MemoryPressureKPIName, MemoryPressureKPIText, FKPIValue::ESeverity::Critical,
 									(int)FGenericPlatformMemoryStats::EMemoryPressureStatus::Nominal, MemoryPressureKPILimit,
 									FKPIValue::LessThan, FKPIValue::Custom);
 	MemoryPressureKPIValue.CustomDisplayValueGetter.BindRaw(this, &FEditorPerformanceModule::GetMemoryPressureDisplayValue);
@@ -247,7 +247,7 @@ void FEditorPerformanceModule::InitializeKPIs()
 	MemoryPressureKPI			= KPIRegistry.DeclareKPIValue(MoveTemp(MemoryPressureKPIValue));
 
 	FKPIValue DynamicInvalidationKPIValue = FKPIValue(EditorCategoryName, EditorCategoryText, DynamicInvalidationKPIName,
-		DynamicInvalidationKPIText, 0.0, DynamicInvalidationKPILimit, FKPIValue::GreaterThanOrEqual, FKPIValue::Custom);
+		DynamicInvalidationKPIText, FKPIValue::ESeverity::Major, 0.0, DynamicInvalidationKPILimit, FKPIValue::GreaterThanOrEqual, FKPIValue::Custom);
 	DynamicInvalidationKPIValue.CustomDisplayValueGetter.BindRaw(this, &FEditorPerformanceModule::GetDynamicInvalidationDisplayValue);
 	DynamicInvalidationKPI		= KPIRegistry.DeclareKPIValue(MoveTemp(DynamicInvalidationKPIValue));
 
@@ -651,9 +651,9 @@ void FEditorPerformanceModule::UpdateKPIs(float InDeltaTime)
 
 	DebugForceKPIIntoBadStateIfNeeded();
 
-	static TArray<FGuid> RecordedKPIEvent;
+	static TArray<FGuid, TInlineAllocator<4>> RecordedKPIEvents;
 	
-	const bool bGoodStateBefore = RecordedKPIEvent.IsEmpty();
+	const TArray<FGuid, TInlineAllocator<4>> RecordedKPIEventsBefore = RecordedKPIEvents;
 
 	// Check for KPIs that have exceeded their value
 	for (FKPIValues::TConstIterator It(GetKPIRegistry().GetKPIValues()); It; ++It)
@@ -662,7 +662,7 @@ void FEditorPerformanceModule::UpdateKPIs(float InDeltaTime)
 
 		if (KPIValue.GetState() == FKPIValue::Bad)
 		{
-			if (RecordedKPIEvent.Find(KPIValue.Id) == INDEX_NONE)
+			if (RecordedKPIEvents.Find(KPIValue.Id) == INDEX_NONE)
 			{
 				// KPI has exceeding the threshold for the first time	
 				const UEditorPerformanceSettings* EditorPerformanceSettings = GetDefault<UEditorPerformanceSettings>();
@@ -683,19 +683,17 @@ void FEditorPerformanceModule::UpdateKPIs(float InDeltaTime)
 				}
 
 				// Add this KPI to the list so we don't send the event again
-				RecordedKPIEvent.Emplace(KPIValue.Id);
+				RecordedKPIEvents.Emplace(KPIValue.Id);
 			}
 		}
 		else
 		{
 			// No longer exceeding threshold, so next time this KPI is exceeded we will record the event
-			RecordedKPIEvent.Remove(KPIValue.Id);
+			RecordedKPIEvents.Remove(KPIValue.Id);
 		}
 	}
 
-	const bool bGoodStateAfter = RecordedKPIEvent.IsEmpty();
-
-	if (bGoodStateBefore != bGoodStateAfter)
+	if (RecordedKPIEventsBefore != RecordedKPIEvents)
 	{
 		OnPerformanceStateChanged.Broadcast();
 	}
diff --git a/Engine/Plugins/Experimental/EditorPerformance/Source/EditorPerformance/Private/KPIValue.cpp b/Engine/Plugins/Experimental/EditorPerformance/Source/EditorPerformance/Private/KPIValue.cpp
--- a/Engine/Plugins/Experimental/EditorPerformance/Source/EditorPerformance/Private/KPIValue.cpp
+++ b/Engine/Plugins/Experimental/EditorPerformance/Source/EditorPerformance/Private/KPIValue.cpp
@@ -244,7 +244,7 @@ FString	FKPIValue::GetValueAsString( float Value, FKPIValue::EDisplayType Displa
 
 FGuid FKPIRegistry::DeclareKPIValue(const FName Category, const FText DisplayCategory, const FName Name, const FText DisplayName, float InitialValue, TOptional<float> ThresholdValue, FKPIValue::ECompare Compare, FKPIValue::EDisplayType Type)
 {
-	return DeclareKPIValue(FKPIValue(Category, DisplayCategory, Name, DisplayName, InitialValue, ThresholdValue, Compare, Type, FKPIValue::EState::NotSet));
+	return DeclareKPIValue(FKPIValue(Category, DisplayCategory, Name, DisplayName, FKPIValue::ESeverity::Minor, InitialValue, ThresholdValue, Compare, Type, FKPIValue::EState::NotSet));
 }
 
 FGuid FKPIRegistry::DeclareKPIValue( const FKPIValue& Value )
diff --git a/Engine/Plugins/Experimental/EditorPerformance/Source/EditorPerformance/Private/SEditorPerformanceStatusBar.cpp b/Engine/Plugins/Experimental/EditorPerformance/Source/EditorPerformance/Private/SEditorPerformanceStatusBar.cpp
--- a/Engine/Plugins/Experimental/EditorPerformance/Source/EditorPerformance/Private/SEditorPerformanceStatusBar.cpp
+++ b/Engine/Plugins/Experimental/EditorPerformance/Source/EditorPerformance/Private/SEditorPerformanceStatusBar.cpp
@@ -63,8 +63,8 @@ void SEditorPerformanceStatusBarWidget::Construct(const FArguments& InArgs)
 {
 	this->ChildSlot
 		[
-			SNew(SButton)
-			.ButtonStyle(&FAppStyle::Get().GetWidgetStyle<FButtonStyle>("StatusBar.StatusBarButton"))
+			SAssignNew(StatusBarButton, SButton)
+			.ButtonStyle(FAppStyle::Get(), "EditorDiagnostics.StatusBarButton")
 			.Content()
 			[
 				SNew(SBox)
@@ -132,34 +132,41 @@ void SEditorPerformanceStatusBarWidget::UpdateState()
 
 		if (KPIValue.GetState() == FKPIValue::Bad)
 		{
-			if (EditorPerformanceSettings)
+			// Turn into a warning state if the severity warrants it
+			if (KPIValue.Severity == FKPIValue::ESeverity::Major)
 			{
-				if (EditorPerformanceSettings->NotifyList.Find(KPIValue.Path) != INDEX_NONE)
-				{
-					// Turn into a warning state if the user wants to be notified	
-					EditorPerformanceState = EEditorPerformanceState::Warnings;
+				EditorPerformanceState = EEditorPerformanceState::Warnings;
+			}
+			else if (KPIValue.Severity == FKPIValue::ESeverity::Critical)
+			{
+				EditorPerformanceState = EEditorPerformanceState::Critical;
+			}
 
-					if (AcknowledgedNotifications.Find(KPIValue.Path) == INDEX_NONE && CurrentNotificationName.IsNone())
-					{
-						NotificationTitle = KPIValue.DisplayName;
+			if (EditorPerformanceSettings && EditorPerformanceSettings->NotifyList.Find(KPIValue.Path) != INDEX_NONE)
+			{
+				// Turn into a warning state if the user wants to be notified
+				EditorPerformanceState = FMath::Max(EEditorPerformanceState::Warnings, EditorPerformanceState);
 
-						FKPIHint Hint;
-						if (EditorPerfModule.GetKPIRegistry().GetKPIHint(KPIValue.Id, Hint))
-						{
-							CurrentNotificationMessage = Hint.Message;
-						}
-						else
-						{
-							CurrentNotificationMessage = FText::FromString(*FString::Printf(TEXT("%s - %s was %s but should be %s %s"),
-								*KPIValue.Category.ToString(),
-								*KPIValue.Name.ToString(),
-								*FKPIValue::GetValueAsString(KPIValue.CurrentValue, KPIValue.DisplayType, KPIValue.CustomDisplayValueGetter),
-								*FKPIValue::GetComparisonAsPrettyString(KPIValue.Compare),
-								*FKPIValue::GetValueAsString(KPIValue.ThresholdValue.GetValue(), KPIValue.DisplayType, KPIValue.CustomDisplayValueGetter)));
-						}
+				if (AcknowledgedNotifications.Find(KPIValue.Path) == INDEX_NONE && CurrentNotificationName.IsNone())
+				{
+					NotificationTitle = KPIValue.DisplayName;
 
-						CurrentNotificationName = KPIValue.Path;
+					FKPIHint Hint;
+					if (EditorPerfModule.GetKPIRegistry().GetKPIHint(KPIValue.Id, Hint))
+					{
+						CurrentNotificationMessage = Hint.Message;
+					}
+					else
+					{
+						CurrentNotificationMessage = FText::FromString(*FString::Printf(TEXT("%s - %s was %s but should be %s %s"),
+							*KPIValue.Category.ToString(),
+							*KPIValue.Name.ToString(),
+							*FKPIValue::GetValueAsString(KPIValue.CurrentValue, KPIValue.DisplayType, KPIValue.CustomDisplayValueGetter),
+							*FKPIValue::GetComparisonAsPrettyString(KPIValue.Compare),
+							*FKPIValue::GetValueAsString(KPIValue.ThresholdValue.GetValue(), KPIValue.DisplayType, KPIValue.CustomDisplayValueGetter)));
 					}
+
+					CurrentNotificationName = KPIValue.Path;
 				}
 			}
 
@@ -271,6 +278,11 @@ void SEditorPerformanceStatusBarWidget::UpdateState()
 		NotificationItem->SetCompletionState(SNotificationItem::CS_None);
 		NotificationItem->ExpireAndFadeout();
 	}
+
+	if (StatusBarButton)
+	{
+		StatusBarButton->SetButtonStyle(GetButtonStyle());
+	}
 }
 
 const FSlateBrush* SEditorPerformanceStatusBarWidget::GetStatusIcon() const
@@ -287,6 +299,10 @@ const FSlateBrush* SEditorPerformanceStatusBarWidget::GetStatusIcon() const
 		{
 			return FAppStyle::Get().GetBrush("EditorDiagnostics.StatusBar.BadgeBG");
 		}
+		case EEditorPerformanceState::Critical:
+		{
+			return FAppStyle::Get().GetBrush("Icons.Alert.Solid");
+		}
 	}
 }
 
@@ -296,6 +312,7 @@ const FSlateBrush* SEditorPerformanceStatusBarWidget::GetStatusBadgeIcon() const
 	{
 		default:
 		case EEditorPerformanceState::Good:
+		case EEditorPerformanceState::Critical:
 		{
 			return FAppStyle::Get().GetBrush("NoBrush");
 		}
@@ -312,4 +329,22 @@ FText SEditorPerformanceStatusBarWidget::GetStatusToolTipText() const
 	return FText::Format(LOCTEXT("ToolTipFormat", "Editor Diagnostics: {0}"), EditorPerformanceStateMessage);
 }
 
+const FButtonStyle* SEditorPerformanceStatusBarWidget::GetButtonStyle() const
+{
+	switch (EditorPerformanceState)
+	{
+		default:
+		case EEditorPerformanceState::Good:
+		case EEditorPerformanceState::Warnings:
+		{
+			return &FAppStyle::Get().GetWidgetStyle<FButtonStyle>("EditorDiagnostics.StatusBarButton");
+		}
+
+		case EEditorPerformanceState::Critical:
+		{
+			return &FAppStyle::Get().GetWidgetStyle<FButtonStyle>("EditorDiagnostics.StatusBarButtonAlert");
+		}
+	}
+}
+
 #undef LOCTEXT_NAMESPACE
diff --git a/Engine/Plugins/Experimental/EditorPerformance/Source/EditorPerformance/Private/SEditorPerformanceStatusBar.h b/Engine/Plugins/Experimental/EditorPerformance/Source/EditorPerformance/Private/SEditorPerformanceStatusBar.h
--- a/Engine/Plugins/Experimental/EditorPerformance/Source/EditorPerformance/Private/SEditorPerformanceStatusBar.h
+++ b/Engine/Plugins/Experimental/EditorPerformance/Source/EditorPerformance/Private/SEditorPerformanceStatusBar.h
@@ -9,8 +9,10 @@
 #include "Widgets/DeclarativeSyntaxSupport.h"
 #include "Widgets/SCompoundWidget.h"
 
+struct FButtonStyle;
 class FUICommandInfo;
 class FUICommandList;
+class SButton;
 class SNotificationItem;
 class SWidget;
 class FKPIValue;
@@ -21,6 +23,7 @@ enum class EEditorPerformanceState : uint8
 {
 	Good,
 	Warnings,
+	Critical
 };
 
 class SEditorPerformanceStatusBarWidget : public SCompoundWidget
@@ -36,6 +39,8 @@ private:
 	const FSlateBrush*				GetStatusIcon() const;
 	const FSlateBrush*				GetStatusBadgeIcon() const;
 	FText							GetStatusToolTipText() const;
+	const FButtonStyle*				GetButtonStyle() const;
+
 	void							UpdateState();
 	
 	TSharedPtr<SNotificationItem>	NotificationItem;
@@ -46,5 +51,6 @@ private:
 	TArray<FName>					AcknowledgedNotifications;
 	uint32							WarningCount = 0;
 
+	TSharedPtr<SButton>				StatusBarButton;
 };
 
diff --git a/Engine/Plugins/Experimental/EditorPerformance/Source/EditorPerformance/Public/KPIValue.h b/Engine/Plugins/Experimental/EditorPerformance/Source/EditorPerformance/Public/KPIValue.h
--- a/Engine/Plugins/Experimental/EditorPerformance/Source/EditorPerformance/Public/KPIValue.h
+++ b/Engine/Plugins/Experimental/EditorPerformance/Source/EditorPerformance/Public/KPIValue.h
@@ -42,12 +42,21 @@ public:
 		Bad,
 	};
 
-	FKPIValue(FName NewCategory, FText NewDisplayCategory, FName NewName, FText NewDisplayName, float NewInitialValue, TOptional<float> NewThresholdValue, FKPIValue::ECompare NewCompare = ECompare::LessThan, FKPIValue::EDisplayType NewDisplayType=EDisplayType::Number, FKPIValue::EState NewState= NotSet):
+	enum ESeverity : uint8
+	{
+		Minor,
+		Major,
+		Critical,
+	};
+
+	FKPIValue(FName NewCategory, FText NewDisplayCategory, FName NewName, FText NewDisplayName, FKPIValue::ESeverity NewSeverity, float NewInitialValue, TOptional<float> NewThresholdValue,
+		FKPIValue::ECompare NewCompare = ECompare::LessThan, FKPIValue::EDisplayType NewDisplayType=EDisplayType::Number, FKPIValue::EState NewState=NotSet):
 		Id(FGuid::NewGuid()),
 		Category(NewCategory),
 		DisplayCategory(NewDisplayCategory),
 		Name(NewName),
 		DisplayName(NewDisplayName),
+		Severity(NewSeverity),
 		CurrentValue(NewInitialValue),
 		ThresholdValue(NewThresholdValue),
 		State(NewState),
@@ -77,6 +86,8 @@ public:
 
 	FName				Path;
 
+	ESeverity			Severity;
+
 	float				CurrentValue = 0;
 	TOptional<float>	ThresholdValue;
 	uint32				FailureCount = 0;
diff --git a/Engine/Source/Editor/EditorStyle/Private/StarshipStyle.cpp b/Engine/Source/Editor/EditorStyle/Private/StarshipStyle.cpp
--- a/Engine/Source/Editor/EditorStyle/Private/StarshipStyle.cpp
+++ b/Engine/Source/Editor/EditorStyle/Private/StarshipStyle.cpp
@@ -1177,6 +1177,19 @@ void FStarshipEditorStyle::FStyle::SetupGeneralStyles()
 
 		Set("EditorDiagnostics.SetShowActivityMonitor", new IMAGE_BRUSH_SVG("Starship/Common/Scalability_16", Icon20x20));
 		Set("EditorDiagnostics.SetShowStalls", new IMAGE_BRUSH_SVG("Starship/Common/Stall", Icon20x20));
+
+		Set("EditorDiagnostics.StatusBarButton", GetWidgetStyle<FButtonStyle>("StatusBar.StatusBarButton"));
+
+		FButtonStyle StatusBarButtonAlert = GetWidgetStyle<FButtonStyle>("StatusBar.StatusBarButton");
+		StatusBarButtonAlert
+					   .SetNormal(FSlateColorBrush(FStyleColors::Error))
+					   .SetNormalForeground(FStyleColors::ForegroundInverted)
+					   .SetHovered(FSlateColorBrush(FStyleColors::Error))
+					   .SetHoveredForeground(FStyleColors::ForegroundHover)
+					   .SetPressed(FSlateColorBrush(FStyleColors::Error))
+					   .SetPressedForeground(FStyleColors::ForegroundHover);
+
+		Set("EditorDiagnostics.StatusBarButtonAlert", StatusBarButtonAlert);
 	}
 
 	// Scene Outliner
