Added Free Disk Space KPIs and removed the free disk space check on editor startup

#rnx
#Synced-CL 49280314

diff --git a/Engine/Config/BaseEditorSettings.ini b/Engine/Config/BaseEditorSettings.ini
--- a/Engine/Config/BaseEditorSettings.ini
+++ b/Engine/Config/BaseEditorSettings.ini
@@ -71,5 +71,8 @@ Duration=(SubFrame=0.0)
 [/Script/UnrealEd.EditorPerformanceSettings]
 bEnableEditorPeformanceTool=True
 bEnableNotifications=True
-NotifyList=Memory_Memory_Pressure
++NotifyList=Memory_Memory_Pressure
++NotifyList=Disk_Engine_Dir_Free_Space
++NotifyList=Disk_Project_Dir_Free_Space
++NotifyList=Disk_User_Dir_Free_Space
 
diff --git a/Engine/Plugins/Experimental/EditorPerformance/Source/EditorPerformance/Private/EditorPerformanceModule.cpp b/Engine/Plugins/Experimental/EditorPerformance/Source/EditorPerformance/Private/EditorPerformanceModule.cpp
--- a/Engine/Plugins/Experimental/EditorPerformance/Source/EditorPerformance/Private/EditorPerformanceModule.cpp
+++ b/Engine/Plugins/Experimental/EditorPerformance/Source/EditorPerformance/Private/EditorPerformanceModule.cpp
@@ -42,6 +42,9 @@ const FText CacheCategoryText = LOCTEXT("Cache", "Cache");
 const FName MemoryCategoryName = TEXT("Memory");
 const FText MemoryCategoryText = LOCTEXT("Memory", "Memory");
 
+const FName DiskCategoryName = TEXT("Disk");
+const FText DiskCategoryText = LOCTEXT("Disk", "Disk");
+
 const FName EditorBootKPIName = TEXT("Boot");
 const FText EditorBootKPIText = LOCTEXT("EditorBoot", "Boot");
 
@@ -105,6 +108,15 @@ const FText MemoryPressureKPIText = LOCTEXT("MemoryPressure", "Memory Pressure")
 const FName DynamicInvalidationKPIName = TEXT("Dynamic Invalidation");
 const FText DynamicInvalidationKPIText = LOCTEXT("DynamicInvalidation", "Dynamic Invalidation");
 
+const FName EngineDirFreeSpaceKPIName = TEXT("Engine Dir Free Space");
+const FText EngineDirFreeSpaceKPIText = LOCTEXT("EngineDirFreeSpace", "Engine Dir Free Space");
+
+const FName ProjectDirFreeSpaceKPIName = TEXT("Project Dir Free Space");
+const FText ProjectDirFreeSpaceKPIText = LOCTEXT("ProjectDirFreeSpace", "Project Dir Free Space");
+
+const FName UserDirFreeSpaceKPIName = TEXT("User Dir Free Space");
+const FText UserDirFreeSpaceKPIText = LOCTEXT("UserDirFreeSpace", "User Dir Free Space");
+
 // Thresholds
 TOptional<float> NoSetLimit;
 float LocalDDCEffciencyKPILimit = 85.f;
@@ -136,10 +148,16 @@ void FEditorPerformanceModule::InitializeEditor()
 {
 	UEditorPerformanceSettings* EditorPerformanceSettings = GetMutableDefault<UEditorPerformanceSettings>();
 
-	FTimerDelegate HeartBeatDelegate;
-	HeartBeatDelegate.BindRaw(this, &FEditorPerformanceModule::HeartBeatCallback);
-	GEditor->GetTimerManager()->SetTimer(HeartBeatTimerHandle, HeartBeatDelegate, HeartBeatIntervalSeconds, true);
+	constexpr bool bLoop = true;
+
+	FTimerDelegate HighFrequencyUpdateDelegate;
+	HighFrequencyUpdateDelegate.BindRaw(this, &FEditorPerformanceModule::UpdateKPIsHighFrequency);
+	GEditor->GetTimerManager()->SetTimer(HighFrequencyUpdateTimerHandle, HighFrequencyUpdateDelegate, HighFrequencyUpdateIntervalSeconds, bLoop);
 	
+	FTimerDelegate LowFrequencyUpdateDelegate;
+	LowFrequencyUpdateDelegate.BindRaw(this, &FEditorPerformanceModule::UpdateKPIsLowFrequency);
+	GEditor->GetTimerManager()->SetTimer(LowFrequencyUpdateTimerHandle, LowFrequencyUpdateDelegate, LowFrequencyUpdateIntervalSeconds, bLoop);
+
 	FTimerDelegate HitchSamplerDelegate;
 	HitchSamplerDelegate.BindRaw(this, &FEditorPerformanceModule::HitchSamplerCallback);
 	GEditor->GetTimerManager()->SetTimer(HitchSamplerTimerHandle, HitchSamplerDelegate, HitchSamplerIntervalSeconds, true);
@@ -171,23 +189,17 @@ void FEditorPerformanceModule::InitializeEditor()
 		EditorPerfSection.AddEntry(FToolMenuEntry::InitWidget("EditorPerformanceStatusBar", CreateStatusBarWidget(), FText::GetEmpty(), true, false));
 	}
 	
-	// Check if we want to have the tool enabled or not.
-	if (EditorPerformanceSettings && EditorPerformanceSettings->bEnableEditorPeformanceTool == false )
+	// Enable any experimental features
+	if (EditorPerformanceSettings && EditorPerformanceSettings->bEnableExperimentalFeatures)
 	{
-		return;
+		CloudDDCLatencyKPI			= KPIRegistry.DeclareKPIValue(CacheCategoryName, CacheCategoryText, CloudDDCLatencyKPIName, CloudDDCLatencyKPIText, 0.0, NoSetLimit, FKPIValue::LessThan, FKPIValue::Milliseconds);
+		CloudDDCReadSpeedKPI		= KPIRegistry.DeclareKPIValue(CacheCategoryName, CacheCategoryText, CloudDDCReadSpeedKPIName, CloudDDCReadSpeedKPIText, 100.0, NoSetLimit, FKPIValue::GreaterThan, FKPIValue::MegaBitsPerSecond);
+		EditorStallRateKPI			= KPIRegistry.DeclareKPIValue(EditorCategoryName, EditorCategoryText, EditorStallRateKPIName, EditorStallRateKPIText, 0.0, NoSetLimit, FKPIValue::LessThan, FKPIValue::Percent);
+		PIEStallRateKPI				= KPIRegistry.DeclareKPIValue(PIECategoryName, PIECategoryText, PIEStallRateKPIName, PIEStallRateKPIText, 0.0, NoSetLimit, FKPIValue::LessThan, FKPIValue::Percent);
 	}
 
-	if (EditorPerformanceSettings)
-	{
-		// Enable any experimental features
-		if (EditorPerformanceSettings->bEnableExperimentalFeatures)
-		{
-			CloudDDCLatencyKPI			= KPIRegistry.DeclareKPIValue(CacheCategoryName, CacheCategoryText, CloudDDCLatencyKPIName, CloudDDCLatencyKPIText, 0.0, NoSetLimit, FKPIValue::LessThan, FKPIValue::Milliseconds);
-			CloudDDCReadSpeedKPI		= KPIRegistry.DeclareKPIValue(CacheCategoryName, CacheCategoryText, CloudDDCReadSpeedKPIName, CloudDDCReadSpeedKPIText, 100.0, NoSetLimit, FKPIValue::GreaterThan, FKPIValue::MegaBitsPerSecond);
-			EditorStallRateKPI			= KPIRegistry.DeclareKPIValue(EditorCategoryName, EditorCategoryText, EditorStallRateKPIName, EditorStallRateKPIText, 0.0, NoSetLimit, FKPIValue::LessThan, FKPIValue::Percent);
-			PIEStallRateKPI				= KPIRegistry.DeclareKPIValue(PIECategoryName, PIECategoryText, PIEStallRateKPIName, PIEStallRateKPIText, 0.0, NoSetLimit, FKPIValue::LessThan, FKPIValue::Percent);
-		}
-	}
+	// Update low frequency KPIs right away so that we have initial values for them without having to way for the timer to tick once
+	UpdateKPIsLowFrequency();
 }
 
 void FEditorPerformanceModule::TerminateEditor()
@@ -202,6 +214,13 @@ void FEditorPerformanceModule::TerminateEditor()
 		}
 	}
 
+	if (GEditor)
+	{
+		GEditor->GetTimerManager()->ClearTimer(HighFrequencyUpdateTimerHandle);
+		GEditor->GetTimerManager()->ClearTimer(LowFrequencyUpdateTimerHandle);
+		GEditor->GetTimerManager()->ClearTimer(HitchSamplerTimerHandle);
+	}
+
 	FEditorDiagnosticsCommands::Unregister();
 }
 
@@ -223,33 +242,62 @@ void FEditorPerformanceModule::ShowPerformanceReportTab()
 void FEditorPerformanceModule::InitializeKPIs()
 {
 	// Declare the KPIs 
-	EditorBootKPI				= KPIRegistry.DeclareKPIValue(EditorCategoryName, EditorCategoryText, EditorBootKPIName, EditorBootKPIText, 0.0, NoSetLimit, FKPIValue::LessThan, FKPIValue::Minutes);
-	EditorInitializeKPI			= //KPIRegistry.DeclareKPIValue(EditorCategoryName, EditorCategoryText, EditorInitializeKPIName, EditorInitializeKPIText, 0.0, NoSetLimit, FKPIValue::LessThan, FKPIValue::Minutes);
-	EditorLoadMapKPI			= KPIRegistry.DeclareKPIValue(EditorCategoryName, EditorCategoryText, EditorLoadMapKPIName, EditorLoadMapKPIText, 0.0, NoSetLimit, FKPIValue::LessThan, FKPIValue::Minutes);
-	EditorAssetRegistryScanKPI	= KPIRegistry.DeclareKPIValue(EditorCategoryName, EditorCategoryText, EditorAssetRegistryScanKPIName, EditorAssetRegistryScanKPIText, 0.0, NoSetLimit, FKPIValue::LessThan, FKPIValue::Minutes);
-	EditorPluginCountKPI		= KPIRegistry.DeclareKPIValue(EditorCategoryName, EditorCategoryText, EditorPluginCountKPIName, EditorPluginCountKPIText, 0.0, NoSetLimit, FKPIValue::LessThan, FKPIValue::Decimal);
-	EditorHitchRateKPI			= KPIRegistry.DeclareKPIValue(EditorCategoryName, EditorCategoryText, EditorHitchRateKPIName, EditorHitchRateKPIText, 0.0, NoSetLimit, FKPIValue::LessThan, FKPIValue::Percent);
-	TotalTimeToEditorKPI		= KPIRegistry.DeclareKPIValue(EditorCategoryName, EditorCategoryText, TotalTimeToEditorKPIName, TotalTimeToEditorKPIText, 0.0, NoSetLimit, FKPIValue::LessThan, FKPIValue::Minutes);
-	PIEFirstTransitionKPI		= KPIRegistry.DeclareKPIValue(PIECategoryName, PIECategoryText, PIEFirstTransitionKPIName, PIEFirstTransitionKPIText, 0.0, NoSetLimit, FKPIValue::LessThan, FKPIValue::Minutes);
-	PIETransitionKPI			= KPIRegistry.DeclareKPIValue(PIECategoryName, PIECategoryText, PIETransitionKPIName, PIETransitionKPIText, 0.0, NoSetLimit, FKPIValue::LessThan, FKPIValue::Minutes);
-	PIEShutdownKPI				= KPIRegistry.DeclareKPIValue(PIECategoryName, PIECategoryText, PIEShutdownKPIName, PIEShutdownKPIText, 0.0, NoSetLimit, FKPIValue::LessThan, FKPIValue::Minutes);
-	PIEHitchRateKPI				= KPIRegistry.DeclareKPIValue(PIECategoryName, PIECategoryText, PIEHitchRateKPIName, PIEHitchRateKPIText, 0.0, NoSetLimit, FKPIValue::LessThan, FKPIValue::Percent);
-	TotalTimeToPIEKPI			= KPIRegistry.DeclareKPIValue(PIECategoryName, PIECategoryText, TotalTimeToPIEKPIName, TotalTimeToPIEKPIText, 0.0, NoSetLimit, FKPIValue::LessThan, FKPIValue::Minutes);
-	LocalDDCEfficiencyKPI		= KPIRegistry.DeclareKPIValue(CacheCategoryName, CacheCategoryText, LocalDDCEfficiencyKPIName, LocalDDCEfficiencyKPIText, 100.0, LocalDDCEffciencyKPILimit, FKPIValue::GreaterThan, FKPIValue::Percent);
-	VirtualAssetEfficiencyKPI	= KPIRegistry.DeclareKPIValue(CacheCategoryName, CacheCategoryText, VirtualAssetEfficiencyKPIName, VirtualAssetEfficiencyKPIText, 100.0, VirtualAssetEfficiencyKPILimit, FKPIValue::GreaterThan, FKPIValue::Percent);
-	AvailableMemoryKPI			= KPIRegistry.DeclareKPIValue(MemoryCategoryName, MemoryCategoryText, AvailableMemoryKPIName, AvailableMemoryKPIText, 128.0, NoSetLimit, FKPIValue::GreaterThanOrEqual, FKPIValue::GigaBytes);
-
-	FKPIValue MemoryPressureKPIValue = FKPIValue(MemoryCategoryName, MemoryCategoryText, MemoryPressureKPIName, MemoryPressureKPIText, FKPIValue::ESeverity::Critical,
-									(int)FGenericPlatformMemoryStats::EMemoryPressureStatus::Nominal, MemoryPressureKPILimit,
-									FKPIValue::LessThan, FKPIValue::Custom);
-	MemoryPressureKPIValue.CustomDisplayValueGetter.BindRaw(this, &FEditorPerformanceModule::GetMemoryPressureDisplayValue);
-
-	MemoryPressureKPI			= KPIRegistry.DeclareKPIValue(MoveTemp(MemoryPressureKPIValue));
-
-	FKPIValue DynamicInvalidationKPIValue = FKPIValue(EditorCategoryName, EditorCategoryText, DynamicInvalidationKPIName,
-		DynamicInvalidationKPIText, FKPIValue::ESeverity::Major, 0.0, DynamicInvalidationKPILimit, FKPIValue::GreaterThanOrEqual, FKPIValue::Custom);
-	DynamicInvalidationKPIValue.CustomDisplayValueGetter.BindRaw(this, &FEditorPerformanceModule::GetDynamicInvalidationDisplayValue);
-	DynamicInvalidationKPI		= KPIRegistry.DeclareKPIValue(MoveTemp(DynamicInvalidationKPIValue));
+	{
+		EditorBootKPI				= KPIRegistry.DeclareKPIValue(EditorCategoryName, EditorCategoryText, EditorBootKPIName, EditorBootKPIText, 0.0, NoSetLimit, FKPIValue::LessThan, FKPIValue::Minutes);
+		EditorInitializeKPI			= //KPIRegistry.DeclareKPIValue(EditorCategoryName, EditorCategoryText, EditorInitializeKPIName, EditorInitializeKPIText, 0.0, NoSetLimit, FKPIValue::LessThan, FKPIValue::Minutes);
+		EditorLoadMapKPI			= KPIRegistry.DeclareKPIValue(EditorCategoryName, EditorCategoryText, EditorLoadMapKPIName, EditorLoadMapKPIText, 0.0, NoSetLimit, FKPIValue::LessThan, FKPIValue::Minutes);
+		EditorAssetRegistryScanKPI	= KPIRegistry.DeclareKPIValue(EditorCategoryName, EditorCategoryText, EditorAssetRegistryScanKPIName, EditorAssetRegistryScanKPIText, 0.0, NoSetLimit, FKPIValue::LessThan, FKPIValue::Minutes);
+		EditorPluginCountKPI		= KPIRegistry.DeclareKPIValue(EditorCategoryName, EditorCategoryText, EditorPluginCountKPIName, EditorPluginCountKPIText, 0.0, NoSetLimit, FKPIValue::LessThan, FKPIValue::Decimal);
+		EditorHitchRateKPI			= KPIRegistry.DeclareKPIValue(EditorCategoryName, EditorCategoryText, EditorHitchRateKPIName, EditorHitchRateKPIText, 0.0, NoSetLimit, FKPIValue::LessThan, FKPIValue::Percent);
+		TotalTimeToEditorKPI		= KPIRegistry.DeclareKPIValue(EditorCategoryName, EditorCategoryText, TotalTimeToEditorKPIName, TotalTimeToEditorKPIText, 0.0, NoSetLimit, FKPIValue::LessThan, FKPIValue::Minutes);
+		FKPIValue DynamicInvalidationKPIValue = FKPIValue(EditorCategoryName, EditorCategoryText, DynamicInvalidationKPIName,
+			DynamicInvalidationKPIText, FKPIValue::ESeverity::Major, 0.0, DynamicInvalidationKPILimit, FKPIValue::GreaterThanOrEqual, FKPIValue::Custom);
+		DynamicInvalidationKPIValue.CustomDisplayValueGetter.BindRaw(this, &FEditorPerformanceModule::GetDynamicInvalidationDisplayValue);
+		DynamicInvalidationKPI		= KPIRegistry.DeclareKPIValue(MoveTemp(DynamicInvalidationKPIValue));
+	}
+	
+	{
+		PIEFirstTransitionKPI		= KPIRegistry.DeclareKPIValue(PIECategoryName, PIECategoryText, PIEFirstTransitionKPIName, PIEFirstTransitionKPIText, 0.0, NoSetLimit, FKPIValue::LessThan, FKPIValue::Minutes);
+		PIETransitionKPI			= KPIRegistry.DeclareKPIValue(PIECategoryName, PIECategoryText, PIETransitionKPIName, PIETransitionKPIText, 0.0, NoSetLimit, FKPIValue::LessThan, FKPIValue::Minutes);
+		PIEShutdownKPI				= KPIRegistry.DeclareKPIValue(PIECategoryName, PIECategoryText, PIEShutdownKPIName, PIEShutdownKPIText, 0.0, NoSetLimit, FKPIValue::LessThan, FKPIValue::Minutes);
+		PIEHitchRateKPI				= KPIRegistry.DeclareKPIValue(PIECategoryName, PIECategoryText, PIEHitchRateKPIName, PIEHitchRateKPIText, 0.0, NoSetLimit, FKPIValue::LessThan, FKPIValue::Percent);
+		TotalTimeToPIEKPI			= KPIRegistry.DeclareKPIValue(PIECategoryName, PIECategoryText, TotalTimeToPIEKPIName, TotalTimeToPIEKPIText, 0.0, NoSetLimit, FKPIValue::LessThan, FKPIValue::Minutes);
+	}
+	
+	{
+		LocalDDCEfficiencyKPI		= KPIRegistry.DeclareKPIValue(CacheCategoryName, CacheCategoryText, LocalDDCEfficiencyKPIName, LocalDDCEfficiencyKPIText, 100.0, LocalDDCEffciencyKPILimit, FKPIValue::GreaterThan, FKPIValue::Percent);
+		VirtualAssetEfficiencyKPI	= KPIRegistry.DeclareKPIValue(CacheCategoryName, CacheCategoryText, VirtualAssetEfficiencyKPIName, VirtualAssetEfficiencyKPIText, 100.0, VirtualAssetEfficiencyKPILimit, FKPIValue::GreaterThan, FKPIValue::Percent);
+	}
+	
+	{
+		AvailableMemoryKPI			= KPIRegistry.DeclareKPIValue(MemoryCategoryName, MemoryCategoryText, AvailableMemoryKPIName, AvailableMemoryKPIText, 128.0, NoSetLimit, FKPIValue::GreaterThanOrEqual, FKPIValue::GigaBytes);
+
+		FKPIValue MemoryPressureKPIValue = FKPIValue(MemoryCategoryName, MemoryCategoryText, MemoryPressureKPIName, MemoryPressureKPIText, FKPIValue::ESeverity::Critical,
+										(int)FGenericPlatformMemoryStats::EMemoryPressureStatus::Nominal, MemoryPressureKPILimit,
+										FKPIValue::LessThan, FKPIValue::Custom);
+		MemoryPressureKPIValue.CustomDisplayValueGetter.BindRaw(this, &FEditorPerformanceModule::GetMemoryPressureDisplayValue);
+
+		MemoryPressureKPI			= KPIRegistry.DeclareKPIValue(MoveTemp(MemoryPressureKPIValue));
+	}
+
+	{
+		FKPIValue EngineDirFreeSpaceKPIValue	= FKPIValue(DiskCategoryName, DiskCategoryText, EngineDirFreeSpaceKPIName, EngineDirFreeSpaceKPIText, FKPIValue::ESeverity::Critical,
+															0.0, 5.0, FKPIValue::GreaterThanOrEqual, FKPIValue::GigaBytes);
+		EngineDirFreeSpaceKPI					= KPIRegistry.DeclareKPIValue(EngineDirFreeSpaceKPIValue);
+
+		FKPIValue ProjectDirFreeSpaceKPIValue	= FKPIValue(DiskCategoryName, DiskCategoryText, ProjectDirFreeSpaceKPIName, ProjectDirFreeSpaceKPIText, FKPIValue::ESeverity::Critical,
+															0.0, 5.0, FKPIValue::GreaterThanOrEqual, FKPIValue::GigaBytes);
+		ProjectDirFreeSpaceKPI					= KPIRegistry.DeclareKPIValue(ProjectDirFreeSpaceKPIValue);
+
+		FKPIValue UserDirFreeSpaceKPIValue		= FKPIValue(DiskCategoryName, DiskCategoryText, UserDirFreeSpaceKPIName, UserDirFreeSpaceKPIText, FKPIValue::ESeverity::Critical,
+															0.0, 1.0, FKPIValue::GreaterThanOrEqual, FKPIValue::GigaBytes);
+		UserDirFreeSpaceKPI						= KPIRegistry.DeclareKPIValue(UserDirFreeSpaceKPIValue);
+
+		const FText FreeSpaceHint = LOCTEXT("FreeSpaceHintMessage", "{0} has limited free space.\n\nTo avoid potential problems, please consider freeing up at least {1} GB.");
+		KPIRegistry.DeclareKPIHint(EngineDirFreeSpaceKPI, FText::Format(FreeSpaceHint, FText::FromString(FPaths::ConvertRelativePathToFull(FPlatformProcess::BaseDir())), EngineDirFreeSpaceKPIValue.ThresholdValue.Get(5.0)));
+		KPIRegistry.DeclareKPIHint(ProjectDirFreeSpaceKPI, FText::Format(FreeSpaceHint, FText::FromString(FPaths::ConvertRelativePathToFull(FPlatformMisc::ProjectDir())), ProjectDirFreeSpaceKPIValue.ThresholdValue.Get(5.0)));
+		KPIRegistry.DeclareKPIHint(UserDirFreeSpaceKPI, FText::Format(FreeSpaceHint, FText::FromString(FPaths::ConvertRelativePathToFull(FPlatformProcess::UserDir())), UserDirFreeSpaceKPIValue.ThresholdValue.Get(1.0)));
+	}
 
 	// Declare the KPI Hints
 	KPIRegistry.DeclareKPIHint(EditorBootKPI, LOCTEXT("EditorBootHintMessage", "The Editor boot time is slow.\nCheck you have enabled a Game Feature Plugin profile for your project and that the expected local cache efficiency is met.\nIf you are booting the Editor in the background then disable the Use Less CPU in Background option in the settings."));
@@ -264,7 +312,7 @@ void FEditorPerformanceModule::InitializeKPIs()
 	KPIRegistry.DeclareKPIHint(MemoryPressureKPI, LOCTEXT("MemoryPressureHintMessage", "Your system is running low on memory.\nTry closing applications that are no longer needed to recover available memory."));
 
 	KPIRegistry.DeclareKPIHint(DynamicInvalidationKPI, LOCTEXT("DynamicInvalidationHintMessage", "Some Dynamic Invalidation Panels contain widgets that do not support invalidation."));
-
+	
 	// Load the KPI profiles
 	KPIRegistry.LoadKPIProfiles(TEXT("EditorPerformance.Profile"), GEditorIni);
 
@@ -445,6 +493,13 @@ void FEditorPerformanceModule::HitchSamplerCallback()
 	static uint32 HitchCount = 0;
 	static uint32 StallCount = 0;
 	
+	const UEditorPerformanceSettings* EditorPerformanceSettings = GetDefault<UEditorPerformanceSettings>();
+
+	if (EditorPerformanceSettings && EditorPerformanceSettings->bEnableEditorPeformanceTool == false)
+	{
+		return;
+	}
+
 	// Only sample framerate and hitches and stalls when we have focus
 	if (FApp::HasFocus())
 	{
@@ -476,14 +531,15 @@ void FEditorPerformanceModule::HitchSamplerCallback()
 	}
 }
 
-void FEditorPerformanceModule::HeartBeatCallback()
+void FEditorPerformanceModule::UpdateKPIsHighFrequency()
 {
-	// Update the KPIS
-	UpdateKPIs(HeartBeatIntervalSeconds);
-}
+	const UEditorPerformanceSettings* EditorPerformanceSettings = GetDefault<UEditorPerformanceSettings>();
+
+	if (EditorPerformanceSettings && EditorPerformanceSettings->bEnableEditorPeformanceTool == false)
+	{
+		return;
+	}
 
-void FEditorPerformanceModule::UpdateKPIs(float InDeltaTime)
-{
 	// Gather memory stats
 	const FPlatformMemoryStats MemoryStats = FPlatformMemory::GetStats();
 	KPIRegistry.SetKPIValue(AvailableMemoryKPI, static_cast<float>(MemoryStats.AvailablePhysical) / (1024.0f * 1024.0f * 1024.0f));
@@ -665,8 +721,6 @@ void FEditorPerformanceModule::UpdateKPIs(float InDeltaTime)
 			if (RecordedKPIEvents.Find(KPIValue.Id) == INDEX_NONE)
 			{
 				// KPI has exceeding the threshold for the first time	
-				const UEditorPerformanceSettings* EditorPerformanceSettings = GetDefault<UEditorPerformanceSettings>();
-
 				if (EditorPerformanceSettings)
 				{
 					if (EditorPerformanceSettings->bEnableSnapshots)
@@ -699,6 +753,39 @@ void FEditorPerformanceModule::UpdateKPIs(float InDeltaTime)
 	}
 }
 
+void FEditorPerformanceModule::UpdateKPIsLowFrequency()
+{
+	const UEditorPerformanceSettings* EditorPerformanceSettings = GetDefault<UEditorPerformanceSettings>();
+
+	if (EditorPerformanceSettings && EditorPerformanceSettings->bEnableEditorPeformanceTool == false)
+	{
+		return;
+	}
+
+	// Gather free disk space
+	{
+		uint64 TotalDiskSpace = 0;
+		uint64 FreeDiskSpace = 0;
+
+		if (FPlatformMisc::GetDiskTotalAndFreeSpace(FPlatformProcess::BaseDir(), TotalDiskSpace, FreeDiskSpace))
+		{
+			KPIRegistry.SetKPIValue(EngineDirFreeSpaceKPI, static_cast<float>(FreeDiskSpace) / (1024.f * 1024.f * 1024.f));
+		}
+
+		if (FPlatformMisc::GetDiskTotalAndFreeSpace(FPlatformMisc::ProjectDir(), TotalDiskSpace, FreeDiskSpace))
+		{
+			KPIRegistry.SetKPIValue(ProjectDirFreeSpaceKPI, static_cast<float>(FreeDiskSpace) / (1024.f * 1024.f * 1024.f));
+		}
+
+		if (FPlatformMisc::GetDiskTotalAndFreeSpace(FPlatformProcess::UserDir(), TotalDiskSpace, FreeDiskSpace))
+		{
+			KPIRegistry.SetKPIValue(UserDirFreeSpaceKPI, static_cast<float>(FreeDiskSpace) / (1024.f * 1024.f * 1024.f));
+		}
+	}
+
+	DebugForceKPIIntoBadStateIfNeeded();
+}
+
 void FEditorPerformanceModule::DebugForceKPIIntoBadStateIfNeeded()
 {
 	if (!GForceBadKPIState.IsEmpty())
diff --git a/Engine/Plugins/Experimental/EditorPerformance/Source/EditorPerformance/Private/KPIValue.cpp b/Engine/Plugins/Experimental/EditorPerformance/Source/EditorPerformance/Private/KPIValue.cpp
--- a/Engine/Plugins/Experimental/EditorPerformance/Source/EditorPerformance/Private/KPIValue.cpp
+++ b/Engine/Plugins/Experimental/EditorPerformance/Source/EditorPerformance/Private/KPIValue.cpp
@@ -205,17 +205,17 @@ FString	FKPIValue::GetValueAsString( float Value, FKPIValue::EDisplayType Displa
 
 		case EDisplayType::Bytes:
 		{
-			return FString::Printf(TEXT("%.2fb"), Value);
+			return FString::Printf(TEXT("%.2fB"), Value);
 		}
 
 		case EDisplayType::MegaBytes:
 		{
-			return FString::Printf(TEXT("%.2fMb"), Value);
+			return FString::Printf(TEXT("%.2fMB"), Value);
 		}
 
 		case EDisplayType::GigaBytes:
 		{
-			return FString::Printf(TEXT("%.2fGb"), Value);
+			return FString::Printf(TEXT("%.2fGB"), Value);
 		}
 
 		case EDisplayType::MegaBitsPerSecond:
diff --git a/Engine/Plugins/Experimental/EditorPerformance/Source/EditorPerformance/Private/SEditorPerformanceDialogs.cpp b/Engine/Plugins/Experimental/EditorPerformance/Source/EditorPerformance/Private/SEditorPerformanceDialogs.cpp
--- a/Engine/Plugins/Experimental/EditorPerformance/Source/EditorPerformance/Private/SEditorPerformanceDialogs.cpp
+++ b/Engine/Plugins/Experimental/EditorPerformance/Source/EditorPerformance/Private/SEditorPerformanceDialogs.cpp
@@ -554,7 +554,23 @@ TSharedRef<SWidget> SEditorPerformanceReportDialog::GetKPIGridPanel()
 							FKPIValue KPIValue;
 							if (EditorPerfModule.GetKPIRegistry().GetKPIValue(KPIId, KPIValue))
 							{
-								return KPIValue.State == FKPIValue::Bad ? FAppStyle::Get().GetBrush("EditorPerformance.Report.Warning") : FAppStyle::GetNoBrush();
+								if (KPIValue.State == FKPIValue::Bad)
+								{
+									switch (KPIValue.Severity)
+									{
+									case FKPIValue::ESeverity::Critical:
+										return FAppStyle::Get().GetBrush("EditorPerformance.Report.Critical");
+									case FKPIValue::ESeverity::Major:
+										return FAppStyle::Get().GetBrush("EditorPerformance.Report.Major");
+									case FKPIValue::ESeverity::Minor:
+									default:
+										return FAppStyle::Get().GetBrush("EditorPerformance.Report.Minor");
+									}
+								}
+								else
+								{
+									return FAppStyle::GetNoBrush();
+								}
 							}
 							else
 							{
diff --git a/Engine/Plugins/Experimental/EditorPerformance/Source/EditorPerformance/Public/EditorPerformanceModule.h b/Engine/Plugins/Experimental/EditorPerformance/Source/EditorPerformance/Public/EditorPerformanceModule.h
--- a/Engine/Plugins/Experimental/EditorPerformance/Source/EditorPerformance/Public/EditorPerformanceModule.h
+++ b/Engine/Plugins/Experimental/EditorPerformance/Source/EditorPerformance/Public/EditorPerformanceModule.h
@@ -49,8 +49,6 @@ public:
 	UE_API TSharedRef<SWidget>	CreateStatusBarWidget();
 	UE_API void ShowPerformanceReportTab();
 
-	UE_API void UpdateKPIs(float DeltaTime);
-
 	UE_API const FKPIRegistry& GetKPIRegistry() const;
 	UE_API const FString& GetKPIProfileName() const;
 
@@ -72,7 +70,8 @@ private:
 	void InitializeKPIs();
 	void TerminateKPIs();
 
-	void HeartBeatCallback();
+	void UpdateKPIsHighFrequency();
+	void UpdateKPIsLowFrequency();
 	void HitchSamplerCallback();
 
 	void DebugForceKPIIntoBadStateIfNeeded();
@@ -100,8 +99,12 @@ private:
 	float							EditorAssetRegistryScanTime = 0;
 	volatile int32					EditorAssetRegistryScanCount = 0;
 	FString							EditorMapName=TEXT("Boot");
-	FTimerHandle					HeartBeatTimerHandle;
-	const float						HeartBeatIntervalSeconds = 1.0f;
+
+	FTimerHandle					HighFrequencyUpdateTimerHandle;
+	const float						HighFrequencyUpdateIntervalSeconds = 1.0f;
+	FTimerHandle					LowFrequencyUpdateTimerHandle;
+	const float						LowFrequencyUpdateIntervalSeconds = 60.0f;
+
 	FTimerHandle					HitchSamplerTimerHandle;
 	const float						HitchSamplerIntervalSeconds = 0.1f;
 	const float						MinFPSForHitching = 15.0f;
@@ -133,6 +136,10 @@ private:
 	FGuid							AvailableMemoryKPI;
 	FGuid							MemoryPressureKPI;
 	FGuid							DynamicInvalidationKPI;
+
+	FGuid							EngineDirFreeSpaceKPI;
+	FGuid							ProjectDirFreeSpaceKPI;
+	FGuid							UserDirFreeSpaceKPI;
 };
 
 
diff --git a/Engine/Source/Editor/EditorStyle/Private/StarshipStyle.cpp b/Engine/Source/Editor/EditorStyle/Private/StarshipStyle.cpp
--- a/Engine/Source/Editor/EditorStyle/Private/StarshipStyle.cpp
+++ b/Engine/Source/Editor/EditorStyle/Private/StarshipStyle.cpp
@@ -1169,7 +1169,9 @@ void FStarshipEditorStyle::FStyle::SetupGeneralStyles()
 	// Editor Diagnostics
 	{
 		Set("EditorPerformance.Report.Panel", new IMAGE_BRUSH_SVG("Starship/Common/Diagnostics_16", Icon16x16));
-		Set("EditorPerformance.Report.Warning", new CORE_IMAGE_BRUSH_SVG("Starship/Common/alert-triangle", Icon16x16, EStyleColor::Warning));
+		Set("EditorPerformance.Report.Minor", new CORE_IMAGE_BRUSH_SVG("Starship/Common/alert-triangle", Icon16x16, EStyleColor::Warning));
+		Set("EditorPerformance.Report.Major", new CORE_IMAGE_BRUSH_SVG("Starship/Common/alert-triangle", Icon16x16, EStyleColor::AccentOrange));
+		Set("EditorPerformance.Report.Critical", new CORE_IMAGE_BRUSH_SVG("Starship/Common/alert-triangle", Icon16x16, EStyleColor::Error));
 
 		Set("EditorDiagnostics.StatusBar.Icon", new IMAGE_BRUSH_SVG("Starship/Common/Diagnostics_16", Icon16x16));
 		Set("EditorDiagnostics.StatusBar.BadgeBG", new IMAGE_BRUSH_SVG("Starship/Common/DiagnosticsBadgeBG_16", Icon16x16));
diff --git a/Engine/Source/Editor/UnrealEd/Classes/Editor/UnrealEdEngine.h b/Engine/Source/Editor/UnrealEd/Classes/Editor/UnrealEdEngine.h
--- a/Engine/Source/Editor/UnrealEd/Classes/Editor/UnrealEdEngine.h
+++ b/Engine/Source/Editor/UnrealEd/Classes/Editor/UnrealEdEngine.h
@@ -1025,12 +1025,6 @@ private:
 	*/
 	UNREALED_API int32 InternalGetNumDirtyPackagesThatNeedCheckout(bool bCheckIfAny) const;
 
-	/**
-	* Internal function to validate free space on drives used by Unreal Engine
-	* The intent is to notify the user of situations where stability maybe impacted.
-	*/
-	UNREALED_API void ValidateFreeDiskSpace() const;
-
 	/** Internal function to filter and add visualizers to a specific list */
 	UNREALED_API void AddVisualizers(AActor* Actor, TArray<FCachedComponentVisualizer>& Visualizers, TFunctionRef<bool(const TSharedPtr<FComponentVisualizer>&)> Condition);
 
diff --git a/Engine/Source/Editor/UnrealEd/Private/UnrealEdEngine.cpp b/Engine/Source/Editor/UnrealEd/Private/UnrealEdEngine.cpp
--- a/Engine/Source/Editor/UnrealEd/Private/UnrealEdEngine.cpp
+++ b/Engine/Source/Editor/UnrealEd/Private/UnrealEdEngine.cpp
@@ -98,9 +98,6 @@ void UUnrealEdEngine::Init(IEngineLoop* InEngineLoop)
 
 	RebuildTemplateMapData();
 
-	// Display warnings to the user about disk space issues
-	ValidateFreeDiskSpace();
-
 	// Build databases used by source code navigation
 	FSourceCodeNavigation::Initialize();
 
@@ -1562,41 +1559,3 @@ void UUnrealEdEngine::OnContentPathDismounted(const FString& AssetPath, const FS
 	MountPointCheckedForWritePermission.Remove(*AssetPath);
 }
 
-bool IsBelowFreeDiskSpaceLimit(const TCHAR* TestDir, FText& OutAppendMessage, const FText& LocationDescriptor, const uint64 MinMB = 5120)
-{
-	uint64 TotalDiskSpace = 0;
-	uint64 FreeDiskSpace = 0;
-
-	if (FPlatformMisc::GetDiskTotalAndFreeSpace(TestDir, TotalDiskSpace, FreeDiskSpace))
-	{
-		const uint64 HardDriveFreeMB = FreeDiskSpace / (1024 * 1024);
-		if (HardDriveFreeMB < MinMB)
-		{
-			static const FText AppendWarning = NSLOCTEXT("DriveSpaceDialog", "LowHardDriveSpaceFormatMsg", "{0}\n\n{1}\n\t\t{2}\n\t\tRecommended: {4} MB\n\t\tFree: {3} MB");
-
-			OutAppendMessage = FText::Format(AppendWarning, OutAppendMessage, LocationDescriptor, FText::FromString(FPaths::ConvertRelativePathToFull(TestDir)), HardDriveFreeMB, MinMB);
-
-			return true;
-		}
-	}
-	return false;
-}
-
-void UUnrealEdEngine::ValidateFreeDiskSpace() const
-{
-	FText Message = NSLOCTEXT("DriveSpaceDialog", "LowHardDriveSpaceMsgHeader", "The following locations have limited free space. To avoid potential problems, please consider freeing up at least the amounts recommended below.");
-	
-	bool bShowWarning = false;
-	bShowWarning |= IsBelowFreeDiskSpaceLimit(FPlatformProcess::BaseDir(), Message, NSLOCTEXT("DriveSpaceDialog", "BaseDirDescriptor", "The base engine directory:"));
-	bShowWarning |= IsBelowFreeDiskSpaceLimit(FPlatformMisc::ProjectDir(), Message, NSLOCTEXT("DriveSpaceDialog", "ProjectDirDescriptor", "The project directory:"));
-	bShowWarning |= IsBelowFreeDiskSpaceLimit(FPlatformProcess::UserDir(), Message, NSLOCTEXT("DriveSpaceDialog", "UserDirDescriptor", "The current user directory, for saving user-specific settings:"), 1024);
-
-	if (bShowWarning)
-	{
-		FEngineAnalytics::LowDriveSpaceDetected();
-
-		const FText Title = NSLOCTEXT("DriveSpaceDialog", "LowHardDriveSpaceMsgTitle", "Warning: Low Drive Space");
-
-		FMessageDialog::Open(EAppMsgType::Ok, Message, Title);
-	}
-}
