Editor Diagnostics
- Add a 5 minutes cooldown to editor KPI failures notifications so that they don't spam when a KPI is near its threshold and repeatedly goes from normal to failure.

#Synced-CL 51753244

diff --git a/Engine/Plugins/Experimental/EditorPerformance/Source/EditorPerformance/Private/SEditorPerformanceStatusBar.cpp b/Engine/Plugins/Experimental/EditorPerformance/Source/EditorPerformance/Private/SEditorPerformanceStatusBar.cpp
--- a/Engine/Plugins/Experimental/EditorPerformance/Source/EditorPerformance/Private/SEditorPerformanceStatusBar.cpp
+++ b/Engine/Plugins/Experimental/EditorPerformance/Source/EditorPerformance/Private/SEditorPerformanceStatusBar.cpp
@@ -53,6 +53,14 @@ struct FSlateBrush;
 #define LOCTEXT_NAMESPACE "EditorPerformance"
 
 
+static float GKPINotificationCooldownSeconds = 300.0f;
+static FAutoConsoleVariableRef CVarKPINotificationCooldownSeconds(
+	TEXT("Editor.PerformanceKPI.NotificationCooldown"),
+	GKPINotificationCooldownSeconds,
+	TEXT("Minimum time in seconds between repeated notifications for the same KPI failure."),
+	ECVF_Default);
+
+
 FReply SEditorPerformanceStatusBarWidget::ViewPerformanceReport_Clicked()
 {
 	FModuleManager::LoadModuleChecked<FEditorPerformanceModule>("EditorPerformance").ShowPerformanceReportTab();
@@ -148,7 +156,10 @@ void SEditorPerformanceStatusBarWidget::UpdateState()
 				// Turn into a warning state if the user wants to be notified
 				EditorPerformanceState = FMath::Max(EEditorPerformanceState::Warnings, EditorPerformanceState);
 
-				if (AcknowledgedNotifications.Find(KPIValue.Path) == INDEX_NONE && CurrentNotificationName.IsNone())
+				const double* LastShownTime = NotificationLastShownTimes.Find(KPIValue.Path);
+				const bool bWithinCooldown = LastShownTime && (FPlatformTime::Seconds() - *LastShownTime < (double)GKPINotificationCooldownSeconds);
+
+				if (AcknowledgedNotifications.Find(KPIValue.Path) == INDEX_NONE && CurrentNotificationName.IsNone() && !bWithinCooldown)
 				{
 					NotificationTitle = KPIValue.DisplayName;
 
@@ -273,6 +284,7 @@ void SEditorPerformanceStatusBarWidget::UpdateState()
 			{
 				AcknowledgeNotificationPromise.SetValue(NotificationItem);
 				NotificationItem->SetCompletionState(SNotificationItem::CS_Fail);
+				NotificationLastShownTimes.Add(CurrentNotificationName, FPlatformTime::Seconds());
 			}
 		}
 	}
diff --git a/Engine/Plugins/Experimental/EditorPerformance/Source/EditorPerformance/Private/SEditorPerformanceStatusBar.h b/Engine/Plugins/Experimental/EditorPerformance/Source/EditorPerformance/Private/SEditorPerformanceStatusBar.h
--- a/Engine/Plugins/Experimental/EditorPerformance/Source/EditorPerformance/Private/SEditorPerformanceStatusBar.h
+++ b/Engine/Plugins/Experimental/EditorPerformance/Source/EditorPerformance/Private/SEditorPerformanceStatusBar.h
@@ -49,6 +49,7 @@ private:
 	FText							CurrentNotificationMessage;
 	FName							CurrentNotificationName;
 	TArray<FName>					AcknowledgedNotifications;
+	TMap<FName, double>				NotificationLastShownTimes;
 	uint32							WarningCount = 0;
 
 	TSharedPtr<SButton>				StatusBarButton;
