[Crash Diagnostics] Add Crash Diagnostics panel to the Editor Diagnostics window

Introduces a new `CrashDiagnostics` module under the `EditorPerformance` plugin. Crash reports from `Saved/Crashes/` are loaded into the Editor Data Storage (TEDS) on startup and displayed in a searchable, sortable table with a details view.

Changes to `CrashReportCore` to export necessary functions.

#Synced-CL 52607881

diff --git a/Engine/Plugins/Experimental/EditorPerformance/EditorPerformance.uplugin b/Engine/Plugins/Experimental/EditorPerformance/EditorPerformance.uplugin
--- a/Engine/Plugins/Experimental/EditorPerformance/EditorPerformance.uplugin
+++ b/Engine/Plugins/Experimental/EditorPerformance/EditorPerformance.uplugin
@@ -24,6 +24,17 @@
 			"Name": "StallLogSubsystem",
 			"Type": "Editor",
 			"LoadingPhase": "PostConfigInit"
+		},
+		{
+			"Name": "CrashDiagnostics",
+			"Type": "Editor",
+			"LoadingPhase": "PostConfigInit"
+		}
+	],
+	"Plugins": [
+		{
+			"Name": "EditorDataStorageFeatures",
+			"Enabled": true
 		}
 	]
 }
\ No newline at end of file
diff --git a/Engine/Plugins/Experimental/EditorPerformance/Source/CrashDiagnostics/CrashDiagnostics.Build.cs b/Engine/Plugins/Experimental/EditorPerformance/Source/CrashDiagnostics/CrashDiagnostics.Build.cs
new file mode 100644
--- /dev/null
+++ b/Engine/Plugins/Experimental/EditorPerformance/Source/CrashDiagnostics/CrashDiagnostics.Build.cs
@@ -0,0 +1,36 @@
+﻿// Copyright Epic Games, Inc. All Rights Reserved.
+
+using UnrealBuildTool;
+
+public class CrashDiagnostics : ModuleRules
+{
+	public CrashDiagnostics(ReadOnlyTargetRules Target) : base(Target)
+	{
+		PCHUsage = ModuleRules.PCHUsageMode.UseExplicitOrSharedPCHs;
+
+		PublicDependencyModuleNames.AddRange(
+			new string[]
+			{
+				"Core",
+				"TypedElementFramework",
+			}
+		);
+
+		PrivateDependencyModuleNames.AddRange(
+			new string[]
+			{
+				"CoreUObject",
+				"CrashReportCore",
+				"InputCore",
+				"Slate",
+				"SlateCore",
+				"TedsQueryStack",
+				"TedsTableViewer",
+				"ToolMenus",
+				"ToolWidgets",
+				"UnrealEd",
+				"XmlParser",
+			}
+		);
+	}
+}
\ No newline at end of file
diff --git a/Engine/Plugins/Experimental/EditorPerformance/Source/CrashDiagnostics/Private/CrashDiagnosticsModule.cpp b/Engine/Plugins/Experimental/EditorPerformance/Source/CrashDiagnostics/Private/CrashDiagnosticsModule.cpp
new file mode 100644
--- /dev/null
+++ b/Engine/Plugins/Experimental/EditorPerformance/Source/CrashDiagnostics/Private/CrashDiagnosticsModule.cpp
@@ -0,0 +1,391 @@
+﻿// Copyright Epic Games, Inc. All Rights Reserved.
+
+#include "CrashDiagnosticsModule.h"
+
+#include "CrashDescription.h"
+#include "CrashDiagnosticsSettings.h"
+#include "Editor.h"
+#include "PlatformErrorReport.h"
+#include "SEditorCrashesPanel.h"
+#include "TedsEditorCrashDataStorageFactory.h"
+#include "Columns/TedsCrashColumns.h"
+#include "DataStorage/Features.h"
+#include "Editor/EditorPerformanceSettings.h"
+#include "Elements/Framework/TypedElementQueryBuilder.h"
+#include "Elements/Interfaces/TypedElementDataStorageInterface.h"
+#include "Framework/Application/SlateApplication.h"
+#include "HAL/FileManager.h"
+#include "Misc/Paths.h"
+#include "Widgets/Images/SLayeredImage.h"
+
+#define LOCTEXT_NAMESPACE "FCrashDiagnosticsModule"
+
+namespace UE::Editor::CrashDiagnostics::Private
+{
+	static FAutoConsoleCommand AddCrashesToDataStorageCommand
+		(
+			TEXT("Editor.CrashDiagnostics.AddCrashesToDataStorage"),
+			TEXT("Retrieve the crashes in the Saved/Crashes folder and add them to the Data Storage"),
+			FConsoleCommandDelegate::CreateLambda([]()
+			{
+				using namespace UE::Editor::DataStorage;
+				if (ICoreProvider* DataStorage = GetMutableDataStorageFeature<ICoreProvider>(StorageFeatureName))
+				{
+					if (FCrashDiagnosticsModule* Module = FCrashDiagnosticsModule::Get())
+					{
+						Module->AddCrashesToDataStorage(*DataStorage);
+					}
+				}
+			})
+		);
+}
+
+
+namespace UE::Editor::CrashDiagnostics
+{
+
+void FCrashDiagnosticsModule::StartupModule()
+{
+	using namespace UE::Editor::DataStorage;
+
+	FEditorDelegates::OnEditorInitialized.AddRaw(this, &FCrashDiagnosticsModule::OnEditorInitialized);
+}
+
+void FCrashDiagnosticsModule::ShutdownModule()
+{
+	FEditorDelegates::OnEditorInitialized.RemoveAll(this);
+}
+
+TArray<TSharedRef<FPrimaryCrashProperties>> FCrashDiagnosticsModule::RetrieveSavedCrashes()
+{
+	TArray<TSharedRef<FPrimaryCrashProperties>> Crashes;
+
+	const FString CrashFolder = FPaths::Combine(*FPaths::ProjectSavedDir(), TEXT("Crashes"));
+	if (!FPaths::DirectoryExists(CrashFolder))
+	{
+		return Crashes;
+	}
+
+	IFileManager::Get().IterateDirectory(*CrashFolder,
+		[&Crashes](const TCHAR* FilenameOrDirectory, bool bIsDirectory)
+		{
+			if (bIsDirectory)
+			{
+				const FPlatformErrorReport ErrorReport(FilenameOrDirectory);
+				FString Filename;
+				// CrashContext.runtime-xml has the precedence over the WER
+				if (ErrorReport.FindFirstReportFileWithExtension(Filename, FGenericCrashContext::CrashContextExtension))
+				{
+					const bool bIsReadOnly = true;
+					Crashes.Add(MakeShared<FCrashContext>(FilenameOrDirectory / Filename, bIsReadOnly));
+				}
+				else if (ErrorReport.FindFirstReportFileWithExtension(Filename, TEXT(".xml")))
+				{
+					const bool bIsReadOnly = true;
+					Crashes.Add(MakeShared<FCrashWERContext>(FilenameOrDirectory / Filename, bIsReadOnly));
+				}
+			}
+			return true;
+		}
+	);
+
+#if 0 // FCrashProperties do not load by default. Uncomment this for debugging
+	for (TSharedRef<FPrimaryCrashProperties>& Crash : Crashes)
+	{
+		Crash->EngineModeEx.AsString();
+		Crash->PlatformFullName.AsString();
+		Crash->CommandLine.AsString();
+		Crash->UserName.AsString();
+		Crash->LoginId.AsString();
+		Crash->EpicAccountId.AsString();
+		Crash->GameSessionID.AsString();
+		Crash->CallStack.AsString();
+		Crash->PCallStack.AsString();
+		Crash->PCallStackHashProperty.AsString();
+		Crash->SourceContext.AsString();
+		Crash->Modules.AsString();
+		Crash->UserDescription.AsString();
+		Crash->UserActivityHint.AsString();
+		Crash->ErrorMessage.AsString();
+		Crash->FullCrashDumpLocation.AsString();
+		Crash->TimeOfCrash.AsInt64();
+		Crash->bAllowToBeContacted.AsBool();
+		Crash->CrashReporterMessage.AsString();
+		Crash->PlatformCallbackResult.AsInt64();
+		Crash->CrashReportClientVersion.AsString();
+		Crash->CPUBrand.AsString();
+	}
+#endif
+
+	return Crashes;
+}
+
+void FCrashDiagnosticsModule::AddCrashesToDataStorage(UE::Editor::DataStorage::ICoreProvider& DataStorage)
+{
+	const TArray<TSharedRef<FPrimaryCrashProperties>> Crashes = RetrieveSavedCrashes();
+	AddCrashesToDataStorage(DataStorage, Crashes);
+}
+
+void FCrashDiagnosticsModule::AddCrashesToDataStorage(UE::Editor::DataStorage::ICoreProvider& DataStorage, const TArray<TSharedRef<FPrimaryCrashProperties>>& Crashes)
+{
+	using namespace UE::Editor::DataStorage;
+	using namespace UE::Editor::CrashDiagnostics;
+
+	// Table is guaranteed to exist, registered by UTedsEditorCrashDataStorageFactory::RegisterTables
+	TableHandle CrashTable = DataStorage.FindTable(DataTableName);
+	check(CrashTable != InvalidTableHandle);
+
+	bool bHadCrashLastSession = false;
+	for (const TSharedRef<FPrimaryCrashProperties>& Crash : Crashes)
+	{
+		FMapKey Key = FMapKey(Crash->CrashGUID);
+		RowHandle Row = DataStorage.LookupMappedRow(MappingDomain, Key);
+		FDateTime CrashTime(Crash->TimeOfCrash.AsInt64());
+
+		if (!DataStorage.IsRowAvailable(Row))
+		{
+			Row = DataStorage.AddRow(CrashTable);
+
+			FString Folder = FPaths::GetPath(Crash->GetXmlFilepath());
+			const FPlatformErrorReport ErrorReport(Folder);
+			FEditorCrashFileReportsColumn CrashReports{ .XmlFilePath = Crash->GetXmlFilepath(), .ReportFilePaths = ErrorReport.GetFilesToUpload() };
+
+			if (FString Filename; ErrorReport.FindFirstReportFileWithExtension(Filename, TEXT(".log")))
+			{
+				CrashReports.LogFilePath = ErrorReport.GetReportDirectory() / MoveTemp(Filename);
+			}
+
+			DataStorage.AddColumn<FEditorCrashFileReportsColumn>(Row, MoveTemp(CrashReports));
+			DataStorage.AddColumn<FEditorCrashGUIDColumn>(Row, { .CrashGUID = Crash->CrashGUID });
+			DataStorage.AddColumn<FEditorCrashTimeColumn>(Row, {
+				.TimeOfCrash = CrashTime,
+				.TimeOfCrashString = FText::AsDateTime(CrashTime, EDateTimeStyle::Default, EDateTimeStyle::Default, FText::GetInvariantTimeZone())
+			});
+			DataStorage.AddColumn<FEditorCrashTypeColumn>(Row, { .CrashType = Crash->CrashType });
+			if (FString ErrorMessage = Crash->ErrorMessage.AsString().TrimStartAndEnd(); !ErrorMessage.IsEmpty())
+			{
+				DataStorage.AddColumn<FEditorCrashErrorMessageColumn>(Row, { .ErrorMessage = MoveTemp(ErrorMessage) });
+			}
+			if (FString CallStack = Crash->CallStack.AsString().TrimStartAndEnd(); !CallStack.IsEmpty())
+			{
+				DataStorage.AddColumn<FEditorCrashCallStackColumn>(Row, { .CallStack = MoveTemp(CallStack) });
+			}
+			if (FString SourceContext = Crash->SourceContext.AsString().TrimStartAndEnd(); !SourceContext.IsEmpty())
+			{
+				DataStorage.AddColumn<FEditorCrashSourceContextColumn>(Row, { .SourceContext = MoveTemp(SourceContext) });
+			}
+			if (FString UserActivity = Crash->UserActivityHint.AsString().TrimStartAndEnd(); !UserActivity.IsEmpty())
+			{
+				DataStorage.AddColumn<FEditorCrashUserActivityHintColumn>(Row, { .UserActivityHint = MoveTemp(UserActivity) });
+			}
+			if (Crash->bIsEnsure)
+			{
+				DataStorage.AddColumn<FEditorCrashIsEnsureTag>(Row);
+			}
+			if (Crash->bIsOOM)
+			{
+				DataStorage.AddColumn<FEditorCrashIsOOMTag>(Row);
+			}
+
+			if (PreviousSessionStartTime != FDateTime() && CrashTime > PreviousSessionStartTime)
+			{
+				bHadCrashLastSession = true;
+				if (!DataStorage.HasColumns<FEditorCrashIsNewTag>(Row))
+				{
+					DataStorage.AddColumn<FEditorCrashIsNewTag>(Row);
+				}
+			}
+
+			DataStorage.MapRow(MappingDomain, Key, Row);
+		}
+	}
+
+	{
+		// Table is guaranteed to exist, registered by UTedsEditorCrashDataStorageFactory::RegisterTables
+		TableHandle GlobalTable = DataStorage.FindTable(GlobalTableName);
+		check(GlobalTable != InvalidTableHandle);
+
+		const FMapKey StateKey(LastSessionCrashStateRowKey);
+		RowHandle StateRow = DataStorage.LookupMappedRow(MappingDomain, StateKey);
+		if (!DataStorage.IsRowAvailable(StateRow))
+		{
+			StateRow = DataStorage.AddRow(GlobalTable);
+			DataStorage.MapRow(MappingDomain, StateKey, StateRow);
+		}
+
+		if (bHadCrashLastSession && !DataStorage.HasColumns<FEditorCrashLastSessionTag>(StateRow))
+		{
+			DataStorage.AddColumn<FEditorCrashLastSessionTag>(StateRow);
+		}
+	}
+}
+
+void FCrashDiagnosticsModule::AddCrashesToDataStorageAsync()
+{
+	AsyncTask(ENamedThreads::AnyBackgroundThreadNormalTask, []()
+	{
+		TArray<TSharedRef<FPrimaryCrashProperties>> Crashes = RetrieveSavedCrashes();
+
+		AsyncTask(ENamedThreads::GameThread, [Crashes = MoveTemp(Crashes)]() mutable
+		{
+			using namespace UE::Editor::DataStorage;
+			using namespace UE::Editor::CrashDiagnostics;
+
+			ICoreProvider* DataStorage = GetMutableDataStorageFeature<ICoreProvider>(StorageFeatureName);
+			FCrashDiagnosticsModule* CrashDiagnosticsModule = Get();
+			if (!CrashDiagnosticsModule || !DataStorage)
+			{
+				return;
+			}
+
+			TPromise<FDelegateHandle> HandlePromise;
+			TSharedFuture<FDelegateHandle> FutureHandle = HandlePromise.GetFuture().Share();
+			FDelegateHandle Handle = DataStorage->OnUpdate().AddLambda([Crashes = MoveTemp(Crashes), FutureHandle]()
+			{
+				ICoreProvider* DataStorage = GetMutableDataStorageFeature<ICoreProvider>(StorageFeatureName);
+				if (DataStorage)
+				{
+					if (FCrashDiagnosticsModule* CrashDiagnosticsModule = Get())
+					{
+						CrashDiagnosticsModule->AddCrashesToDataStorage(*DataStorage, Crashes);
+					}
+					DataStorage->OnUpdate().Remove(FutureHandle.Get());
+				}
+			});
+			HandlePromise.SetValue(Handle);
+		});
+	});
+}
+
+void FCrashDiagnosticsModule::LoadCrashDiagnosticsSettings()
+{
+	GetMutableDefault<UCrashDiagnosticsSettings>()->LoadConfig();
+}
+
+void FCrashDiagnosticsModule::SaveCrashDiagnosticsSettings()
+{
+	GetMutableDefault<UCrashDiagnosticsSettings>()->SaveConfig();
+}
+
+bool FCrashDiagnosticsModule::HasCrashedLastSession() const
+{
+	using namespace UE::Editor::DataStorage;
+	using namespace UE::Editor::DataStorage::Queries;
+	ICoreProvider* DataStorage = GetMutableDataStorageFeature<ICoreProvider>(StorageFeatureName);
+	if (!DataStorage)
+	{
+		return false;
+	}
+
+	const UTedsEditorCrashDataStorageFactory* Factory = DataStorage->FindFactory<UTedsEditorCrashDataStorageFactory>();
+	const QueryHandle Handle = ensure(Factory) ? Factory->GetHasCrashedLastSessionQueryHandle() : InvalidQueryHandle;
+	const FQueryResult Result = Handle != InvalidQueryHandle ? DataStorage->RunQuery(Handle) : FQueryResult();
+	return Result.Count > 0;
+}
+
+bool FCrashDiagnosticsModule::HasUnreadCrashes() const
+{
+	using namespace UE::Editor::DataStorage;
+	using namespace UE::Editor::DataStorage::Queries;
+	ICoreProvider* DataStorage = GetMutableDataStorageFeature<ICoreProvider>(StorageFeatureName);
+	if (!DataStorage)
+	{
+		return false;
+	}
+
+	const UTedsEditorCrashDataStorageFactory* Factory = DataStorage->FindFactory<UTedsEditorCrashDataStorageFactory>();
+	const QueryHandle Handle = ensure(Factory) ? Factory->GetHasUnreadCrashesQueryHandle() : InvalidQueryHandle;
+	const FQueryResult Result = Handle != InvalidQueryHandle ? DataStorage->RunQuery(Handle) : FQueryResult();
+	return Result.Count > 0;
+}
+
+TSharedRef<SWidget> FCrashDiagnosticsModule::CreateCrashLogPanel()
+{
+	return SNew(UE::Editor::CrashDiagnostics::SEditorCrashesPanel);
+}
+
+template <typename T>
+TSharedPtr<T> FindFirstDescendant(const TSharedRef<SWidget>& Root, FName RealWidgetClassName = NAME_None)
+{
+	if (Root->GetType() == (RealWidgetClassName.IsNone() ? T::StaticWidgetClass().GetWidgetType() : RealWidgetClassName))
+	{
+		return StaticCastSharedRef<T>(Root);
+	}
+	if (FChildren* Children = Root->GetChildren())
+	{
+		for (int32 i = 0; i < Children->Num(); ++i)
+		{
+			if (TSharedPtr<T> Result = FindFirstDescendant<T>(Children->GetChildAt(i), RealWidgetClassName))
+			{
+				return Result;
+			}
+		}
+	}
+	return nullptr;
+}
+
+TSharedRef<SWidget> FCrashDiagnosticsModule::DecorateButtonWithUnreadBadge(TSharedRef<SWidget> InWidget)
+{
+	FCrashDiagnosticsModule* Module = Get();
+	if (!Module)
+	{
+		return InWidget;
+	}
+
+	// SLayeredImage does not have a class
+	TSharedPtr<SLayeredImage> LayeredImage = FindFirstDescendant<SLayeredImage>(InWidget, "SLayeredImage");
+	if (LayeredImage)
+	{
+		LayeredImage->AddLayer(TAttribute<const FSlateBrush*>::CreateLambda([]()
+		{
+			const FCrashDiagnosticsModule* Module = Get();
+			return Module ? Module->GetUnreadBadgeIcon() : nullptr;
+		}));
+	}
+	return InWidget;
+}
+
+const FSlateBrush* FCrashDiagnosticsModule::GetUnreadBadgeIcon() const
+{
+	return HasUnreadCrashes() ? FAppStyle::Get().GetBrush("Icons.BadgeModified") : nullptr;
+}
+
+void FCrashDiagnosticsModule::OnEditorInitialized(double TimeToInitializeEditor)
+{
+	LoadCrashDiagnosticsSettings();
+	{
+		UCrashDiagnosticsSettings* Settings = GetMutableDefault<UCrashDiagnosticsSettings>();
+		PreviousSessionStartTime = Settings->LastSessionStartTime;
+		Settings->LastSessionStartTime = FDateTime::UtcNow();
+		Settings->SaveConfig();
+	}
+
+	{
+		using namespace UE::Editor::DataStorage;
+		auto StartupOperations = []()
+		{
+			if (FCrashDiagnosticsModule* CrashDiagnosticsModule = Get())
+			{
+				CrashDiagnosticsModule->AddCrashesToDataStorageAsync();
+			}
+		};
+
+		if (AreEditorDataStorageFeaturesEnabled())
+		{
+			StartupOperations();
+		}
+		else
+		{
+			OnEditorDataStorageFeaturesEnabled().AddLambda([StartupOperations = MoveTemp(StartupOperations)]()
+			{
+				StartupOperations();
+			});
+		}
+	}
+}
+
+}
+
+#undef LOCTEXT_NAMESPACE
+
+IMPLEMENT_MODULE(UE::Editor::CrashDiagnostics::FCrashDiagnosticsModule, CrashDiagnostics)
\ No newline at end of file
diff --git a/Engine/Plugins/Experimental/EditorPerformance/Source/CrashDiagnostics/Private/CrashDiagnosticsSettings.h b/Engine/Plugins/Experimental/EditorPerformance/Source/CrashDiagnostics/Private/CrashDiagnosticsSettings.h
new file mode 100644
--- /dev/null
+++ b/Engine/Plugins/Experimental/EditorPerformance/Source/CrashDiagnostics/Private/CrashDiagnosticsSettings.h
@@ -0,0 +1,21 @@
+﻿// Copyright Epic Games, Inc. All Rights Reserved.
+
+#pragma once
+
+#include "CoreMinimal.h"
+#include "UObject/Object.h"
+#include "CrashDiagnosticsSettings.generated.h"
+
+UCLASS(MinimalAPI, config=EditorPerProjectUserSettings)
+class UCrashDiagnosticsSettings : public UObject
+{
+	GENERATED_BODY()
+
+public:
+	/**
+	 * The start time of the last editor session.
+	 * Used to detect crashes that occurred during the previous session.
+	 */
+	UPROPERTY(config)
+	FDateTime LastSessionStartTime = FDateTime();
+};
diff --git a/Engine/Plugins/Experimental/EditorPerformance/Source/CrashDiagnostics/Private/SEditorCrashesPanel.cpp b/Engine/Plugins/Experimental/EditorPerformance/Source/CrashDiagnostics/Private/SEditorCrashesPanel.cpp
new file mode 100644
--- /dev/null
+++ b/Engine/Plugins/Experimental/EditorPerformance/Source/CrashDiagnostics/Private/SEditorCrashesPanel.cpp
@@ -0,0 +1,302 @@
+﻿// Copyright Epic Games, Inc. All Rights Reserved.
+
+#include "SEditorCrashesPanel.h"
+
+#include "CrashDescription.h"
+#include "ISettingsModule.h"
+#include "TedsEditorCrashDataStorageFactory.h"
+#include "TedsQueryNode.h"
+#include "TedsQuerySearchNode.h"
+#include "ToolMenus.h"
+#include "Columns/TedsCrashColumns.h"
+#include "DataStorage/Features.h"
+#include "Elements/Framework/TypedElementAttributeBinding.h"
+#include "Elements/Framework/TypedElementQueryBuilder.h"
+#include "Framework/Docking/TabManager.h"
+#include "GenericPlatform/GenericPlatformCrashContext.h"
+#include "TedsTableViewer/Public/Widgets/STedsSearchBox.h"
+#include "Widgets/SBoxPanel.h"
+#include "Widgets/SRowDetails.h"
+#include "Widgets/STedsTableViewer.h"
+#include "Widgets/Input/SButton.h"
+#include "Widgets/Input/SComboButton.h"
+#include "Widgets/Layout/SScrollBox.h"
+#include "Widgets/Text/STextBlock.h"
+#include "Widgets/Views/SListView.h"
+
+#define LOCTEXT_NAMESPACE "EditorCrashes"
+
+namespace UE::Editor::CrashDiagnostics
+{
+
+const FName SEditorCrashesPanel::SettingsMenuName = "EditorCrashesSettings";
+
+void SEditorCrashesPanel::Construct(const FArguments& InArgs)
+{
+	using namespace UE::Editor::DataStorage;
+	using namespace UE::Editor::DataStorage::Queries;
+
+	this->ChildSlot
+	[
+		SNew(SSplitter)
+		.Orientation(Orient_Vertical)
+		+ SSplitter::Slot()
+		[
+			CreateTableWidget()
+		]
+		+ SSplitter::Slot()
+		[
+			CreateDetailsWidget()
+		]
+	];
+}
+
+TSharedRef<SWidget> SEditorCrashesPanel::CreateTableWidget()
+{
+	using namespace UE::Editor::DataStorage;
+	using namespace UE::Editor::DataStorage::Queries;
+
+	ICoreProvider* DataStorage = GetMutableDataStorageFeature<ICoreProvider>(StorageFeatureName);
+	check(DataStorage);
+	TSharedRef<QueryStack::FQueryNode> QueryNode = MakeShared<QueryStack::FQueryNode>(*DataStorage,
+		Select()
+		.ReadOnly<
+			FEditorCrashGUIDColumn,
+			FEditorCrashTimeColumn,
+			FEditorCrashErrorMessageColumn,
+			FEditorCrashTypeColumn
+		>()
+		.Where()
+			.None<FEditorCrashIsEnsureTag>()
+		.Compile());
+
+	SearchBox = SNew(STedsSearchBox)
+		.InSearchableQueryNode(QueryNode)
+		.InSearchableQueryFlags(QueryStack::FQuerySearchNode::ESyncActions::RefreshOnUpdate)
+		.OutSearchNode(&SearchNode);
+
+	TSharedPtr<FMetaData> GenericMetaData = MakeShared<FMetaData>();
+	GenericMetaData->AddOrSetMutableData("bUseDisplayNameTooltip", true);
+	TSharedPtr<FColumnMetaData> ColumnMetaData = MakeShared<FColumnMetaData>(FEditorCrashTimeColumn::StaticStruct(), FColumnMetaData::EFlags::None);
+	ColumnMetaData->AddOrSetMutableData("bColumnInitialSortIsAscending", false);
+
+	return SNew(SBorder)
+		.BorderImage(FAppStyle::Get().GetBrush("Brushes.Recessed"))
+		.Padding(0)
+		[
+			SNew(SScrollBox)
+			.Orientation(Orient_Vertical)
+			+ SScrollBox::Slot()
+			.HAlign(HAlign_Fill)
+			[
+				SNew(SVerticalBox)
+				+ SVerticalBox::Slot()
+				.Padding(5, 5, 5, 0)
+				.AutoHeight()
+				[
+					SNew(SHorizontalBox)
+					+ SHorizontalBox::Slot()
+					.FillWidth(1.0f)
+					[
+						SearchBox.ToSharedRef()
+					]
+					+ SHorizontalBox::Slot()
+					.HAlign(HAlign_Right)
+					.Padding(5, 0, 0, 0)
+					.AutoWidth()
+					[
+						SNew(SComboButton)
+						.ComboButtonStyle(FAppStyle::Get(), "SimpleComboButtonWithIcon")
+						.HasDownArrow(false)
+						.OnGetMenuContent(this, &SEditorCrashesPanel::CreateSettingsMenuWidget)
+						.ButtonContent()
+						[
+							SNew(SImage)
+							.Image(FAppStyle::GetBrush("Icons.Settings"))
+							.ColorAndOpacity(FSlateColor::UseForeground())
+						]
+					]
+				]
+				+ SVerticalBox::Slot()
+				.Padding(0, 5, 0, 0)
+				.AutoHeight()
+				[
+					SNew(STedsTableViewer)
+					.QueryStack(SearchNode)
+					.HeaderWidgetPurpose(UTedsEditorCrashDataStorageFactory::GetHeaderRowPurposeID())
+					.CellWidgetPurpose(UTedsEditorCrashDataStorageFactory::GetTableRowPurposeID())
+					.OnSelectionChanged(this, &SEditorCrashesPanel::OnSelectionChanged)
+					.PrimaryColumn(FEditorCrashTimeColumn::StaticStruct())
+					.GenericMetaData(GenericMetaData)
+					.ColumnMetaData({MoveTemp(ColumnMetaData)})
+					.Columns({
+						FEditorCrashTimeColumn::StaticStruct(),
+						FEditorCrashTypeColumn::StaticStruct(),
+						FEditorCrashErrorMessageColumn::StaticStruct(),
+					})
+				]
+			]
+		];
+}
+
+TSharedRef<SWidget> SEditorCrashesPanel::CreateDetailsWidget()
+{
+	using namespace UE::Editor::DataStorage;
+
+	static const TArray<TWeakObjectPtr<const UScriptStruct>> Rows = {
+		FEditorCrashGUIDColumn::StaticStruct(),
+		FEditorCrashTimeColumn::StaticStruct(),
+		FEditorCrashTypeColumn::StaticStruct(),
+		FEditorCrashFileReportsColumn::StaticStruct(),
+		FEditorCrashErrorMessageColumn::StaticStruct(),
+		FEditorCrashUserActivityHintColumn::StaticStruct(),
+		FEditorCrashCallStackColumn::StaticStruct(),
+		FEditorCrashSourceContextColumn::StaticStruct(),
+	};
+
+	TSharedRef<FMetaData> NameMetaData = MakeShared<FMetaData>();
+	NameMetaData->AddOrSetMutableData("Font", TEXT("NormalFontBold"));
+	TSharedRef<FMetaData> DataMetaData = MakeShared<FMetaData>();
+	DataMetaData->AddOrSetMutableData("bUseEditableText", true);
+
+	return SNew(SVerticalBox)
+		+ SVerticalBox::Slot()
+		.AutoHeight()
+		.MinHeight(20.f)
+		[
+			SNew(SBorder)
+			.BorderImage(&FAppStyle::Get().GetWidgetStyle<FTableColumnHeaderStyle>("TableView.Header.Column").NormalBrush)
+			[
+				SNew(SHorizontalBox)
+				+ SHorizontalBox::Slot()
+				.FillWidth(1.0f)
+				.VAlign(VAlign_Center)
+				.Padding(FMargin(4.0f, 0.0f, 0.0f, 0.0f))
+				[
+					SNew(STextBlock)
+					.Text(LOCTEXT("DetailsTitle", "Details"))
+				]
+			]
+		]
+		+ SVerticalBox::Slot()
+		.FillHeight(1.0f)
+		[
+			SNew(SBorder)
+			.BorderImage(&FAppStyle::Get().GetWidgetStyle<FTableRowStyle>("TableView.Row").EvenRowBackgroundBrush)
+			[
+				SNew(SScrollBox)
+				.Orientation(Orient_Vertical)
+				+ SScrollBox::Slot()
+				.Padding(4.0f)
+				[
+					SAssignNew(CrashRowDetails, SRowDetails)
+					.SelectionMode(ESelectionMode::None)
+					.Columns(Rows)
+					.HeaderRow
+					(
+						SNew(SHeaderRow)
+						.Visibility(EVisibility::Collapsed)
+						+ SHeaderRow::Column(SRowDetails::NameColumn)
+							.ManualWidth(100.f)
+						+ SHeaderRow::Column(SRowDetails::DataColumn)
+					)
+					.RowPadding(FMargin(0.f, 3.f))
+					.NameMetaData(NameMetaData)
+					.DataMetaData(DataMetaData)
+				]
+			]
+		];
+}
+
+TSharedRef<SWidget> SEditorCrashesPanel::CreateSettingsMenuWidget()
+{
+	UToolMenus* ToolMenus = UToolMenus::Get();
+	if (!ToolMenus)
+	{
+		return SNullWidget::NullWidget;
+	}
+
+	UToolMenu* SettingsMenu = ToolMenus->FindMenu(SettingsMenuName);
+
+	if (!SettingsMenu)
+	{
+		SettingsMenu = RegisterSettingsMenu();
+	}
+
+	if (SettingsMenu)
+	{
+		return ToolMenus->GenerateWidget(SettingsMenu);
+	}
+	else
+	{
+		return SNullWidget::NullWidget;
+	}
+}
+
+void SEditorCrashesPanel::OnSelectionChanged(UE::Editor::DataStorage::RowHandle RowHandle, ESelectInfo::Type SelectionType)
+{
+	SelectedRowHandle = RowHandle;
+	CrashRowDetails->SetRow(RowHandle);
+
+	if (NewCrashTimer)
+	{
+		UnRegisterActiveTimer(NewCrashTimer.ToSharedRef());
+		NewCrashTimer.Reset();
+	}
+
+	using namespace UE::Editor::DataStorage;
+	ICoreProvider* DataStorage = GetMutableDataStorageFeature<ICoreProvider>(StorageFeatureName);
+
+	if (!ensure(DataStorage) || !DataStorage->IsRowAssigned(RowHandle))
+	{
+		return;
+	}
+
+	if (DataStorage->HasColumns(RowHandle, TArray{FEditorCrashIsNewTag::StaticStruct()}))
+	{
+		NewCrashTimer = RegisterActiveTimer(MarkAsReadDelaySeconds, FWidgetActiveTimerDelegate::CreateLambda(
+			[this, RowHandle](double InCurrentTime, float InDeltaTime)
+			{
+				ICoreProvider* DataStorage = GetMutableDataStorageFeature<ICoreProvider>(StorageFeatureName);
+				if (SelectedRowHandle == RowHandle && DataStorage)
+				{
+					DataStorage->RemoveColumn(RowHandle, FEditorCrashIsNewTag::StaticStruct());
+				}
+				NewCrashTimer.Reset();
+				return EActiveTimerReturnType::Stop;
+			}));
+	}
+}
+
+UToolMenu* SEditorCrashesPanel::RegisterSettingsMenu()
+{
+	UToolMenus* ToolMenus = UToolMenus::Get();
+	if (!ToolMenus)
+	{
+		return nullptr;
+	}
+
+	UToolMenu* SettingsMenu = ToolMenus->RegisterMenu(SettingsMenuName);
+
+	FToolMenuSection& SettingsMenuSection = SettingsMenu->AddSection("Settings", LOCTEXT("Settings", "Settings"));
+
+	SettingsMenuSection.AddEntry(
+		FToolMenuEntry::InitMenuEntry("AllSettings",
+			LOCTEXT("OpenSettingsText", "More Editor Performance Settings..."),
+			LOCTEXT("OpenSettingsToolTip", "Open the Editor Performance Settings Tab."),
+			FSlateIcon(),
+			FUIAction(
+				FExecuteAction::CreateLambda([]()
+				{
+					FModuleManager::LoadModuleChecked<ISettingsModule>("Settings").ShowViewer("Editor", "General", "EditorPerformanceSettings");
+				}),
+				FCanExecuteAction()
+			)));
+
+	return SettingsMenu;
+
+}
+
+}
+
+#undef LOCTEXT_NAMESPACE
\ No newline at end of file
diff --git a/Engine/Plugins/Experimental/EditorPerformance/Source/CrashDiagnostics/Private/SEditorCrashesPanel.h b/Engine/Plugins/Experimental/EditorPerformance/Source/CrashDiagnostics/Private/SEditorCrashesPanel.h
new file mode 100644
--- /dev/null
+++ b/Engine/Plugins/Experimental/EditorPerformance/Source/CrashDiagnostics/Private/SEditorCrashesPanel.h
@@ -0,0 +1,56 @@
+﻿// Copyright Epic Games, Inc. All Rights Reserved.
+
+#pragma once
+
+#include "ToolMenu.h"
+#include "DataStorage/Handles.h"
+#include "Widgets/SCompoundWidget.h"
+
+namespace UE::Editor::DataStorage
+{
+	class SRowDetails;
+
+	namespace QueryStack
+	{
+		class IRowNode;
+	}
+
+	class STedsSearchBox;
+}
+
+#define UE_API CRASHDIAGNOSTICS_API
+
+namespace UE::Editor::CrashDiagnostics
+{
+	class SEditorCrashesPanel : public SCompoundWidget
+	{
+		static const FName SettingsMenuName;
+
+	public:
+		SLATE_BEGIN_ARGS(SEditorCrashesPanel) {}
+		SLATE_END_ARGS()
+
+		UE_API void Construct(const FArguments& InArgs);
+
+	private:
+		TSharedRef<SWidget> CreateTableWidget();
+		TSharedRef<SWidget> CreateDetailsWidget();
+		TSharedRef<SWidget> CreateSettingsMenuWidget();
+
+		void OnSelectionChanged(UE::Editor::DataStorage::RowHandle RowHandle, ESelectInfo::Type SelectionType);
+
+		UToolMenu* RegisterSettingsMenu();
+
+	private:
+		TSharedPtr<UE::Editor::DataStorage::QueryStack::IRowNode> SearchNode;
+
+		TSharedPtr<UE::Editor::DataStorage::STedsSearchBox> SearchBox;
+		TSharedPtr<UE::Editor::DataStorage::SRowDetails> CrashRowDetails;
+
+		UE::Editor::DataStorage::RowHandle SelectedRowHandle = UE::Editor::DataStorage::InvalidRowHandle;
+		float MarkAsReadDelaySeconds = 1.0f;
+		TSharedPtr<FActiveTimerHandle> NewCrashTimer;
+	};
+}
+
+#undef UE_API
diff --git a/Engine/Plugins/Experimental/EditorPerformance/Source/CrashDiagnostics/Private/SFileHyperlink.cpp b/Engine/Plugins/Experimental/EditorPerformance/Source/CrashDiagnostics/Private/SFileHyperlink.cpp
new file mode 100644
--- /dev/null
+++ b/Engine/Plugins/Experimental/EditorPerformance/Source/CrashDiagnostics/Private/SFileHyperlink.cpp
@@ -0,0 +1,100 @@
+﻿// Copyright Epic Games, Inc. All Rights Reserved.
+
+#include "SFileHyperlink.h"
+
+#include "Misc/Paths.h"
+#include "Widgets/Input/SHyperlink.h"
+
+
+#define LOCTEXT_NAMESPACE "TedsCrashColumnWidgets"
+
+namespace UE::Editor::CrashDiagnostics
+{
+
+void SFileHyperlink::Construct(const FArguments& InArgs, const FString& InPath)
+{
+	FullPath = FPaths::ConvertRelativePathToFull(InPath);
+	OpenAction = InArgs._OpenAction;
+
+	ChildSlot
+		.HAlign(HAlign_Left)
+		[
+			SNew(SHyperlink)
+			.Text(InArgs._TextOverride.IsSet()
+				? InArgs._TextOverride
+				: FText::FromString(FPaths::GetCleanFilename(FullPath)))
+			.ToolTipText( this, &SFileHyperlink::GetToolTipText)
+			.OnNavigate(this, &SFileHyperlink::Open)
+		];
+}
+
+SFileHyperlink::EOpenAction SFileHyperlink::GetAction() const
+{
+	if (OpenAction == EOpenAction::None)
+	{
+		return EOpenAction::None;
+	}
+
+	EOpenAction Action = EOpenAction::None;
+	if (FPaths::FileExists(FullPath))
+	{
+		if (OpenAction == EOpenAction::Auto || OpenAction == EOpenAction::OpenFileApplication)
+		{
+			Action = EOpenAction::OpenFileApplication;
+		}
+		else if (OpenAction == EOpenAction::OpenFolder)
+		{
+			Action = EOpenAction::OpenFolder;
+		}
+	}
+	else if (FPaths::DirectoryExists(FullPath))
+	{
+		if (OpenAction == EOpenAction::Auto || OpenAction == EOpenAction::OpenFolder || OpenAction == EOpenAction::OpenFileApplication)
+		{
+			Action = EOpenAction::OpenFolder;
+		}
+	}
+
+	return Action;
+}
+
+void SFileHyperlink::Open() const
+{
+	EOpenAction Action = GetAction();
+	if (Action == EOpenAction::None)
+	{
+		return;
+	}
+
+	if (Action == EOpenAction::OpenFileApplication)
+	{
+		if (!FPlatformProcess::LaunchFileInDefaultExternalApplication(*FullPath))
+		{
+			Action = EOpenAction::OpenFolder;
+		}
+	}
+	if (Action == EOpenAction::OpenFolder)
+	{
+		FPlatformProcess::ExploreFolder(*FullPath);
+	}
+}
+
+FText SFileHyperlink::GetToolTipText() const
+{
+	const EOpenAction Action = GetAction();
+	if (Action == EOpenAction::None)
+	{
+		return FText::FromString(FullPath);
+	}
+
+	return FText::Format(
+		Action == EOpenAction::OpenFileApplication
+		? LOCTEXT("OpenFileInDefaultApplication", "Open file in default application\n{0}")
+		: LOCTEXT("OpenFolder", "Open folder\n{0}"),
+		FText::FromString(FullPath)
+	);
+}
+
+}
+
+#undef LOCTEXT_NAMESPACE
\ No newline at end of file
diff --git a/Engine/Plugins/Experimental/EditorPerformance/Source/CrashDiagnostics/Private/SFileHyperlink.h b/Engine/Plugins/Experimental/EditorPerformance/Source/CrashDiagnostics/Private/SFileHyperlink.h
new file mode 100644
--- /dev/null
+++ b/Engine/Plugins/Experimental/EditorPerformance/Source/CrashDiagnostics/Private/SFileHyperlink.h
@@ -0,0 +1,39 @@
+﻿// Copyright Epic Games, Inc. All Rights Reserved.
+
+#pragma once
+
+#include "Widgets/SCompoundWidget.h"
+
+namespace UE::Editor::CrashDiagnostics
+{
+	class SFileHyperlink : public SCompoundWidget
+	{
+	public:
+		enum class EOpenAction : uint8
+		{
+			None,
+			OpenFileApplication,
+			OpenFolder,
+			Auto,
+		};
+
+		SLATE_BEGIN_ARGS(SFileHyperlink)
+			{}
+			SLATE_ATTRIBUTE(FText, TextOverride)
+			SLATE_ARGUMENT_DEFAULT(EOpenAction, OpenAction) { EOpenAction::Auto };
+		SLATE_END_ARGS()
+
+		void Construct(const FArguments& InArgs, const FString& InPath);
+
+	private:
+		EOpenAction GetAction() const;
+
+		void Open() const;
+
+		FText GetToolTipText() const;
+
+	private:
+		FString FullPath;
+		EOpenAction OpenAction = EOpenAction::Auto;
+	};
+}
diff --git a/Engine/Plugins/Experimental/EditorPerformance/Source/CrashDiagnostics/Private/TedsCrashColumnWidgets.cpp b/Engine/Plugins/Experimental/EditorPerformance/Source/CrashDiagnostics/Private/TedsCrashColumnWidgets.cpp
new file mode 100644
--- /dev/null
+++ b/Engine/Plugins/Experimental/EditorPerformance/Source/CrashDiagnostics/Private/TedsCrashColumnWidgets.cpp
@@ -0,0 +1,446 @@
+﻿// Copyright Epic Games, Inc. All Rights Reserved.
+
+#include "TedsCrashColumnWidgets.h"
+
+#include "CrashDiagnosticsModule.h"
+#include "SFileHyperlink.h"
+#include "TedsEditorCrashDataStorageFactory.h"
+#include "Columns/SlateHeaderColumns.h"
+#include "Columns/TedsCrashColumns.h"
+#include "Elements/Framework/TypedElementAttributeBinding.h"
+#include "Styling/StyleColors.h"
+#include "Widgets/SBoxPanel.h"
+#include "Widgets/Input/SEditableText.h"
+#include "Widgets/Input/SMultiLineEditableTextBox.h"
+#include "Widgets/Layout/SBox.h"
+#include "Widgets/Text/STextBlock.h"
+
+
+#define LOCTEXT_NAMESPACE "TedsCrashColumnWidgets"
+
+///////////////////////////////////
+// FEditorCrashGUIDWidgetConstructor
+///////////////////////////////////
+
+TSharedPtr<SWidget> FEditorCrashGUIDWidgetConstructor::CreateWidget(
+	UE::Editor::DataStorage::ICoreProvider* DataStorage,
+	UE::Editor::DataStorage::IUiProvider* DataStorageUi,
+	UE::Editor::DataStorage::RowHandle TargetRow,
+	UE::Editor::DataStorage::RowHandle WidgetRow,
+	const UE::Editor::DataStorage::FMetaDataView& Arguments)
+{
+	UE::Editor::DataStorage::FAttributeBinder Binder(TargetRow, DataStorage);
+
+	const FEditorCrashGUIDColumn* CrashGUID = DataStorage->GetColumn<FEditorCrashGUIDColumn>(TargetRow);
+	const FEditorCrashFileReportsColumn* FileReports = DataStorage->GetColumn<FEditorCrashFileReportsColumn>(TargetRow);
+	if (!ensure(CrashGUID) || !ensure(FileReports))
+	{
+		return SNullWidget::NullWidget;
+	}
+
+	using namespace UE::Editor::CrashDiagnostics;
+	return SNew(SHorizontalBox)
+		+ SHorizontalBox::Slot()
+		.Padding(FMargin(4.0f, 0.0f, 0.0f, 0.0f))
+		[
+			SNew(SFileHyperlink, FileReports->GetCrashReportFolder())
+				.TextOverride(Binder.BindText(&FEditorCrashGUIDColumn::CrashGUID))
+				.OpenAction(SFileHyperlink::EOpenAction::OpenFolder)
+		];
+}
+
+
+///////////////////////////////////
+// FEditorCrashTimeWidgetConstructor
+///////////////////////////////////
+
+class FEditorCrashTimeWidgetSorter final : public UE::Editor::DataStorage::TColumnSorterInterface<UE::Editor::DataStorage::FColumnSorterInterface::ESortType::FixedSize64, FEditorCrashTimeColumn>
+{
+public:
+	virtual FText GetShortName() const override
+	{
+		return LOCTEXT("FEditorCrashTimeWidgetSorter", "Editor Crash Time Sorter");
+	}
+protected:
+	virtual UE::Editor::DataStorage::FPrefixInfo CalculatePrefix(const FEditorCrashTimeColumn& Column, uint32 ByteIndex) const override
+	{
+		return UE::Editor::DataStorage::CreateSortPrefix(ByteIndex, Column.TimeOfCrash.GetTicks());
+	}
+};
+
+
+TSharedPtr<SWidget> FEditorCrashTimeWidgetConstructor::CreateWidget(
+	UE::Editor::DataStorage::ICoreProvider* DataStorage,
+	UE::Editor::DataStorage::IUiProvider* DataStorageUi,
+	UE::Editor::DataStorage::RowHandle TargetRow,
+	UE::Editor::DataStorage::RowHandle WidgetRow,
+	const UE::Editor::DataStorage::FMetaDataView& Arguments)
+{
+	UE::Editor::DataStorage::FAttributeBinder Binder(TargetRow, DataStorage);
+
+	const FEditorCrashTimeColumn* CrashTime = DataStorage->GetColumn<FEditorCrashTimeColumn>(TargetRow);
+	if (!ensure(CrashTime))
+	{
+		return SNullWidget::NullWidget;
+	}
+
+	TSharedPtr<SWidget> Widget;
+	TAttribute<FText> TextAttribute = Binder.BindData(&FEditorCrashTimeColumn::TimeOfCrashString);
+
+	const bool* bUseEditableText = Arguments.FindForColumnOrGeneric<FEditorCrashTimeColumn>("bUseEditableText").TryGetExact<bool>();
+	if (bUseEditableText && (*bUseEditableText))
+	{
+		Widget = SNew(SEditableText)
+			.IsReadOnly(true)
+			.Text(MoveTemp(TextAttribute));
+	}
+	else
+	{
+		Widget = SNew(STextBlock)
+			.Text(MoveTemp(TextAttribute));
+	}
+
+	return SNew(SHorizontalBox)
+		+ SHorizontalBox::Slot()
+		.Padding(FMargin(4.0f, 0.0f, 0.0f, 0.0f))
+		[
+			Widget.ToSharedRef()
+		];
+}
+
+TArray<TSharedPtr<const UE::Editor::DataStorage::FColumnSorterInterface>> FEditorCrashTimeWidgetConstructor::ConstructColumnSorters(
+	UE::Editor::DataStorage::ICoreProvider* DataStorage,
+	UE::Editor::DataStorage::IUiProvider* DataStorageUi,
+	const UE::Editor::DataStorage::FMetaDataView& Arguments)
+{
+	using namespace UE::Editor::DataStorage;
+	return TArray<TSharedPtr<const FColumnSorterInterface>>(
+		{
+			MakeShared<FEditorCrashTimeWidgetSorter>()
+		});
+}
+
+
+///////////////////////////////////
+// FEditorCrashTimeTableRowWidgetConstructor
+///////////////////////////////////
+
+TSharedPtr<SWidget> FEditorCrashTimeTableRowWidgetConstructor::CreateWidget(UE::Editor::DataStorage::ICoreProvider* DataStorage, UE::Editor::DataStorage::IUiProvider* DataStorageUi, UE::Editor::DataStorage::RowHandle TargetRow, UE::Editor::DataStorage::RowHandle WidgetRow, const UE::Editor::DataStorage::FMetaDataView& Arguments)
+{
+	UE::Editor::DataStorage::FAttributeBinder Binder(TargetRow, DataStorage);
+
+	const FEditorCrashTimeColumn* CrashTime = DataStorage->GetColumn<FEditorCrashTimeColumn>(TargetRow);
+	const FEditorCrashGUIDColumn* CrashGUID = DataStorage->GetColumn<FEditorCrashGUIDColumn>(TargetRow);
+	if (!ensure(CrashTime) || !ensure(CrashGUID))
+	{
+		return SNullWidget::NullWidget;
+	}
+
+	return SNew(SHorizontalBox)
+		+ SHorizontalBox::Slot()
+		.Padding(FMargin(4.0f, 0.0f, 0.0f, 0.0f))
+		[
+			SNew(STextBlock)
+			.Font_Lambda([DataStorage, TargetRow]()
+			{
+				const bool IsNewCrash = DataStorage->HasColumns(TargetRow, TArray{FEditorCrashIsNewTag::StaticStruct()});
+				return IsNewCrash ?
+					FAppStyle::Get().GetFontStyle("NormalFontBold") :
+					FAppStyle::Get().GetFontStyle("NormalText");
+			})
+			.Text(Binder.BindData(&FEditorCrashTimeColumn::TimeOfCrashString, [DataStorage, TargetRow](const FText& InTimeOfCrash) mutable
+			{
+				const bool IsNewCrash = DataStorage->HasColumns(TargetRow, TArray{FEditorCrashIsNewTag::StaticStruct()});
+
+				return FText::Format(LOCTEXT("NewFmt", "{0}{1}"),
+					IsNewCrash ? INVTEXT("*") : FText::GetEmpty(),
+					InTimeOfCrash
+				);
+			}))
+			.ToolTipText(Binder.BindText(&FEditorCrashGUIDColumn::CrashGUID))
+		];
+}
+
+TArray<TSharedPtr<const UE::Editor::DataStorage::FColumnSorterInterface>> FEditorCrashTimeTableRowWidgetConstructor::ConstructColumnSorters(UE::Editor::DataStorage::ICoreProvider* DataStorage, UE::Editor::DataStorage::IUiProvider* DataStorageUi, const UE::Editor::DataStorage::FMetaDataView& Arguments)
+{
+	using namespace UE::Editor::DataStorage;
+	return TArray<TSharedPtr<const FColumnSorterInterface>>(
+		{
+			MakeShared<FEditorCrashTimeWidgetSorter>()
+		});
+}
+
+
+///////////////////////////////////
+// FEditorCrashErrorMessageWidgetConstructor
+///////////////////////////////////
+
+TSharedPtr<SWidget> FEditorCrashErrorMessageWidgetConstructor::CreateWidget(UE::Editor::DataStorage::ICoreProvider* DataStorage, UE::Editor::DataStorage::IUiProvider* DataStorageUi, UE::Editor::DataStorage::RowHandle TargetRow, UE::Editor::DataStorage::RowHandle WidgetRow, const UE::Editor::DataStorage::FMetaDataView& Arguments)
+{
+	UE::Editor::DataStorage::FAttributeBinder Binder(TargetRow, DataStorage);
+
+	const FEditorCrashErrorMessageColumn* ErrorMessage = DataStorage->GetColumn<FEditorCrashErrorMessageColumn>(TargetRow);
+	if (!ensure(ErrorMessage))
+	{
+		return SNullWidget::NullWidget;
+	}
+
+	return SNew(SMultiLineEditableTextBox)
+		.IsReadOnly(true)
+		.Style(UTedsEditorCrashDataStorageFactory::GetReadOnlyTextBoxStyle())
+		.BackgroundColor(FStyleColors::Background)
+		.Padding(FMargin(3.0f, 0.f))
+		.Text(Binder.BindText(&FEditorCrashErrorMessageColumn::ErrorMessage));
+}
+
+TSharedPtr<SWidget> FEditorCrashErrorMessageCompactWidgetConstructor::CreateWidget(
+	UE::Editor::DataStorage::ICoreProvider* DataStorage,
+	UE::Editor::DataStorage::IUiProvider* DataStorageUi,
+	UE::Editor::DataStorage::RowHandle TargetRow,
+	UE::Editor::DataStorage::RowHandle WidgetRow,
+	const UE::Editor::DataStorage::FMetaDataView& Arguments)
+{
+	UE::Editor::DataStorage::FAttributeBinder Binder(TargetRow, DataStorage);
+
+	const FEditorCrashErrorMessageColumn* ErrorMessage = DataStorage->GetColumn<FEditorCrashErrorMessageColumn>(TargetRow);
+	if (!ensure(ErrorMessage))
+	{
+		return SNullWidget::NullWidget;
+	}
+
+	return SNew(SHorizontalBox)
+		+ SHorizontalBox::Slot()
+		.HAlign(HAlign_Left)
+		.MaxWidth(0.f)
+		.AutoWidth()
+		[
+			SNew(STextBlock) // The goal of this empty TextBlock is to give this HBox a height of 1 line of text
+		]
+		+ SHorizontalBox::Slot()
+		.Padding(FMargin(4.0f, 0.0f, 0.0f, 0.0f))
+		.VAlign(VAlign_Fill)
+		.HAlign(HAlign_Fill)
+		.FillWidth(1.0f)
+		[
+			SNew(SBox)
+			.HeightOverride(0.f) // This will make the Text below take only one line with ellipsis
+			[
+				SNew(STextBlock)
+				.Text(Binder.BindText(&FEditorCrashErrorMessageColumn::ErrorMessage))
+				.OverflowPolicy(ETextOverflowPolicy::MultilineEllipsis)
+				.ToolTipText(Binder.BindText(&FEditorCrashErrorMessageColumn::ErrorMessage))
+			]
+		];
+}
+
+
+///////////////////////////////////
+// FEditorCrashTypeWidgetConstructor
+///////////////////////////////////
+
+TSharedPtr<SWidget> FEditorCrashTypeWidgetConstructor::CreateWidget(
+	UE::Editor::DataStorage::ICoreProvider* DataStorage,
+	UE::Editor::DataStorage::IUiProvider* DataStorageUi,
+	UE::Editor::DataStorage::RowHandle TargetRow,
+	UE::Editor::DataStorage::RowHandle WidgetRow,
+	const UE::Editor::DataStorage::FMetaDataView& Arguments)
+{
+	UE::Editor::DataStorage::FAttributeBinder Binder(TargetRow, DataStorage);
+
+	const FEditorCrashTypeColumn* CrashType = DataStorage->GetColumn<FEditorCrashTypeColumn>(TargetRow);
+	if (!ensure(CrashType))
+	{
+		return SNullWidget::NullWidget;
+	}
+
+	TSharedPtr<SWidget> Widget;
+	TAttribute<FText> TextAttribute = Binder.BindData(&FEditorCrashTypeColumn::CrashType, [](const FString& InCrashType) mutable
+	{
+		return FText::FromString(FName::NameToDisplayString(InCrashType, false));
+	});
+
+	const bool* bUseEditableText = Arguments.FindForColumnOrGeneric<FEditorCrashTypeColumn>("bUseEditableText").TryGetExact<bool>();
+	if (bUseEditableText && (*bUseEditableText))
+	{
+		Widget = SNew(SEditableText)
+			.IsReadOnly(true)
+			.Text(MoveTemp(TextAttribute));
+	}
+	else
+	{
+		Widget = SNew(STextBlock)
+			.Text(MoveTemp(TextAttribute));
+	}
+
+	return SNew(SHorizontalBox)
+		+ SHorizontalBox::Slot()
+		.Padding(FMargin(4.0f, 0.0f, 0.0f, 0.0f))
+		[
+			Widget.ToSharedRef()
+		];
+}
+
+
+///////////////////////////////////
+// FEditorCrashFileReportsWidgetConstructor
+///////////////////////////////////
+
+TSharedPtr<SWidget> FEditorCrashFileReportsWidgetConstructor::CreateWidget(UE::Editor::DataStorage::ICoreProvider* DataStorage, UE::Editor::DataStorage::IUiProvider* DataStorageUi, UE::Editor::DataStorage::RowHandle TargetRow, UE::Editor::DataStorage::RowHandle WidgetRow, const UE::Editor::DataStorage::FMetaDataView& Arguments)
+{
+	using namespace UE::Editor::CrashDiagnostics;
+
+	const FEditorCrashFileReportsColumn* FileReports = DataStorage->GetColumn<FEditorCrashFileReportsColumn>(TargetRow);
+	if (!ensure(FileReports))
+	{
+		return SNullWidget::NullWidget;
+	}
+
+	TSharedRef<SVerticalBox> VerticalBox = SNew(SVerticalBox);
+	for (const FString& ReportFileName : FileReports->ReportFilePaths)
+	{
+		VerticalBox->AddSlot()
+		.Padding(0.f, 0.f, 0.f, 3.f)
+		[
+			SNew(SFileHyperlink, ReportFileName)
+				.OpenAction(SFileHyperlink::EOpenAction::OpenFileApplication)
+		];
+	}
+
+	return SNew(SHorizontalBox)
+		+ SHorizontalBox::Slot()
+		.HAlign(HAlign_Fill)
+		.VAlign(VAlign_Fill)
+		.Padding(FMargin(4.0f, 0.0f, 0.0f, 0.0f))
+		[
+			MoveTemp(VerticalBox)
+		];
+}
+
+
+///////////////////////////////////
+// FEditorCrashCallStackWidgetConstructor
+///////////////////////////////////
+
+TSharedPtr<SWidget> FEditorCrashCallStackWidgetConstructor::CreateWidget(UE::Editor::DataStorage::ICoreProvider* DataStorage, UE::Editor::DataStorage::IUiProvider* DataStorageUi, UE::Editor::DataStorage::RowHandle TargetRow, UE::Editor::DataStorage::RowHandle WidgetRow, const UE::Editor::DataStorage::FMetaDataView& Arguments)
+{
+	UE::Editor::DataStorage::FAttributeBinder Binder(TargetRow, DataStorage);
+
+	const FEditorCrashCallStackColumn* CallStack = DataStorage->GetColumn<FEditorCrashCallStackColumn>(TargetRow);
+	if (!ensure(CallStack))
+	{
+		return SNullWidget::NullWidget;
+	}
+
+	return SNew(SMultiLineEditableTextBox)
+		.IsReadOnly(true)
+		.Style(UTedsEditorCrashDataStorageFactory::GetReadOnlyTextBoxStyle())
+		.BackgroundColor(FStyleColors::Background)
+		.Padding(FMargin(3.0f, 0.f))
+		.Text(Binder.BindText(&FEditorCrashCallStackColumn::CallStack));
+}
+
+
+///////////////////////////////////
+// FEditorCrashSourceContextWidgetConstructor
+///////////////////////////////////
+
+TSharedPtr<SWidget> FEditorCrashSourceContextWidgetConstructor::CreateWidget(UE::Editor::DataStorage::ICoreProvider* DataStorage, UE::Editor::DataStorage::IUiProvider* DataStorageUi, UE::Editor::DataStorage::RowHandle TargetRow, UE::Editor::DataStorage::RowHandle WidgetRow, const UE::Editor::DataStorage::FMetaDataView& Arguments)
+{
+	UE::Editor::DataStorage::FAttributeBinder Binder(TargetRow, DataStorage);
+
+	const FEditorCrashSourceContextColumn* SourceContext = DataStorage->GetColumn<FEditorCrashSourceContextColumn>(TargetRow);
+	if (!ensure(SourceContext))
+	{
+		return SNullWidget::NullWidget;
+	}
+
+	return SNew(SMultiLineEditableTextBox)
+		.IsReadOnly(true)
+		.Style(UTedsEditorCrashDataStorageFactory::GetReadOnlyTextBoxStyle())
+		.BackgroundColor(FStyleColors::Background)
+		.Padding(FMargin(2.0f, 0.f))
+		.Text(Binder.BindText(&FEditorCrashSourceContextColumn::SourceContext));
+}
+
+
+///////////////////////////////////
+// FEditorCrashUserActivityHintWidgetConstructor
+///////////////////////////////////
+
+TSharedPtr<SWidget> FEditorCrashUserActivityHintWidgetConstructor::CreateWidget(UE::Editor::DataStorage::ICoreProvider* DataStorage, UE::Editor::DataStorage::IUiProvider* DataStorageUi, UE::Editor::DataStorage::RowHandle TargetRow, UE::Editor::DataStorage::RowHandle WidgetRow, const UE::Editor::DataStorage::FMetaDataView& Arguments)
+{
+	UE::Editor::DataStorage::FAttributeBinder Binder(TargetRow, DataStorage);
+
+	const FEditorCrashUserActivityHintColumn* ActivityHint = DataStorage->GetColumn<FEditorCrashUserActivityHintColumn>(TargetRow);
+	if (!ensure(ActivityHint))
+	{
+		return SNullWidget::NullWidget;
+	}
+
+	TSharedPtr<SWidget> Widget;
+	TAttribute<FText> TextAttribute = Binder.BindText(&FEditorCrashUserActivityHintColumn::UserActivityHint);
+
+	const bool* bUseEditableText = Arguments.FindForColumnOrGeneric<FEditorCrashUserActivityHintColumn>("bUseEditableText").TryGetExact<bool>();
+	if (bUseEditableText && (*bUseEditableText))
+	{
+		Widget = SNew(SEditableText)
+			.IsReadOnly(true)
+			.Text(MoveTemp(TextAttribute));
+	}
+	else
+	{
+		Widget = SNew(STextBlock)
+			.Text(MoveTemp(TextAttribute));
+	}
+
+	return SNew(SHorizontalBox)
+		+ SHorizontalBox::Slot()
+		.Padding(FMargin(4.0f, 0.0f, 0.0f, 0.0f))
+		[
+			Widget.ToSharedRef()
+		];
+}
+
+
+///////////////////////////////////
+// FEditorCrashTimeHeaderWidgetConstructor
+///////////////////////////////////
+
+TSharedPtr<SWidget> FEditorCrashTimeHeaderWidgetConstructor::CreateWidget(const UE::Editor::DataStorage::FMetaDataView& Arguments)
+{
+	return SNew(STextBlock)
+		.Text(FEditorCrashTimeColumn::StaticStruct()->GetDisplayNameText());
+}
+
+bool FEditorCrashTimeHeaderWidgetConstructor::FinalizeWidget(UE::Editor::DataStorage::ICoreProvider* DataStorage, UE::Editor::DataStorage::IUiProvider* DataStorageUi, UE::Editor::DataStorage::RowHandle Row, const TSharedPtr<SWidget>& Widget)
+{
+	DataStorage->AddColumn(Row, FHeaderWidgetSizeColumn
+		{
+			.ColumnSizeMode = EColumnSizeMode::Manual,
+			.Width = 200.0f
+		});
+	return true;
+}
+
+
+///////////////////////////////////
+// FEditorCrashTypeHeaderWidgetConstructor
+///////////////////////////////////
+
+TSharedPtr<SWidget> FEditorCrashTypeHeaderWidgetConstructor::CreateWidget(const UE::Editor::DataStorage::FMetaDataView& Arguments)
+{
+	return SNew(STextBlock)
+		.Text(FEditorCrashTypeColumn::StaticStruct()->GetDisplayNameText());
+}
+
+bool FEditorCrashTypeHeaderWidgetConstructor::FinalizeWidget(UE::Editor::DataStorage::ICoreProvider* DataStorage, UE::Editor::DataStorage::IUiProvider* DataStorageUi, UE::Editor::DataStorage::RowHandle Row, const TSharedPtr<SWidget>& Widget)
+{
+	DataStorage->AddColumn(Row, FHeaderWidgetSizeColumn
+		{
+			.ColumnSizeMode = EColumnSizeMode::Manual,
+			.Width = 150.0f
+		});
+	return true;
+}
+
+#undef LOCTEXT_NAMESPACE
diff --git a/Engine/Plugins/Experimental/EditorPerformance/Source/CrashDiagnostics/Private/TedsCrashColumnWidgets.h b/Engine/Plugins/Experimental/EditorPerformance/Source/CrashDiagnostics/Private/TedsCrashColumnWidgets.h
new file mode 100644
--- /dev/null
+++ b/Engine/Plugins/Experimental/EditorPerformance/Source/CrashDiagnostics/Private/TedsCrashColumnWidgets.h
@@ -0,0 +1,266 @@
+﻿// Copyright Epic Games, Inc. All Rights Reserved.
+
+#pragma once
+
+#include "Elements/Interfaces/TypedElementDataStorageFactory.h"
+#include "Elements/Interfaces/TypedElementDataStorageUiInterface.h"
+
+#include "TedsCrashColumnWidgets.generated.h"
+
+
+///////////////////////////////////
+// Widget Constructors
+///////////////////////////////////
+
+USTRUCT()
+struct FEditorCrashGUIDWidgetConstructor : public FSimpleWidgetConstructor
+{
+	GENERATED_BODY()
+
+public:
+	FEditorCrashGUIDWidgetConstructor()
+		: Super(StaticStruct())
+	{
+	}
+
+	virtual TSharedPtr<SWidget> CreateWidget(
+		UE::Editor::DataStorage::ICoreProvider* DataStorage,
+		UE::Editor::DataStorage::IUiProvider* DataStorageUi,
+		UE::Editor::DataStorage::RowHandle TargetRow,
+		UE::Editor::DataStorage::RowHandle WidgetRow,
+		const UE::Editor::DataStorage::FMetaDataView& Arguments) override;
+};
+
+
+USTRUCT()
+struct FEditorCrashTimeWidgetConstructor : public FSimpleWidgetConstructor
+{
+	GENERATED_BODY()
+
+public:
+	FEditorCrashTimeWidgetConstructor()
+		: Super(StaticStruct())
+	{
+	}
+
+	virtual TSharedPtr<SWidget> CreateWidget(
+		UE::Editor::DataStorage::ICoreProvider* DataStorage,
+		UE::Editor::DataStorage::IUiProvider* DataStorageUi,
+		UE::Editor::DataStorage::RowHandle TargetRow,
+		UE::Editor::DataStorage::RowHandle WidgetRow,
+		const UE::Editor::DataStorage::FMetaDataView& Arguments) override;
+
+	virtual TArray<TSharedPtr<const UE::Editor::DataStorage::FColumnSorterInterface>> ConstructColumnSorters(
+		UE::Editor::DataStorage::ICoreProvider* DataStorage,
+		UE::Editor::DataStorage::IUiProvider* DataStorageUi,
+		const UE::Editor::DataStorage::FMetaDataView& Arguments) override;
+};
+
+USTRUCT()
+struct FEditorCrashTimeTableRowWidgetConstructor : public FSimpleWidgetConstructor
+{
+	GENERATED_BODY()
+
+public:
+	FEditorCrashTimeTableRowWidgetConstructor()
+		: Super(StaticStruct())
+	{
+	}
+
+	virtual TSharedPtr<SWidget> CreateWidget(
+		UE::Editor::DataStorage::ICoreProvider* DataStorage,
+		UE::Editor::DataStorage::IUiProvider* DataStorageUi,
+		UE::Editor::DataStorage::RowHandle TargetRow,
+		UE::Editor::DataStorage::RowHandle WidgetRow,
+		const UE::Editor::DataStorage::FMetaDataView& Arguments) override;
+
+	virtual TArray<TSharedPtr<const UE::Editor::DataStorage::FColumnSorterInterface>> ConstructColumnSorters(
+		UE::Editor::DataStorage::ICoreProvider* DataStorage,
+		UE::Editor::DataStorage::IUiProvider* DataStorageUi,
+		const UE::Editor::DataStorage::FMetaDataView& Arguments) override;
+};
+
+
+USTRUCT()
+struct FEditorCrashErrorMessageWidgetConstructor : public FSimpleWidgetConstructor
+{
+	GENERATED_BODY()
+
+public:
+	FEditorCrashErrorMessageWidgetConstructor()
+		: Super(StaticStruct())
+	{
+	}
+
+	virtual TSharedPtr<SWidget> CreateWidget(
+		UE::Editor::DataStorage::ICoreProvider* DataStorage,
+		UE::Editor::DataStorage::IUiProvider* DataStorageUi,
+		UE::Editor::DataStorage::RowHandle TargetRow,
+		UE::Editor::DataStorage::RowHandle WidgetRow,
+		const UE::Editor::DataStorage::FMetaDataView& Arguments) override;
+};
+
+USTRUCT()
+struct FEditorCrashErrorMessageCompactWidgetConstructor : public FSimpleWidgetConstructor
+{
+	GENERATED_BODY()
+
+public:
+	FEditorCrashErrorMessageCompactWidgetConstructor()
+		: Super(StaticStruct())
+	{
+	}
+
+	virtual TSharedPtr<SWidget> CreateWidget(
+		UE::Editor::DataStorage::ICoreProvider* DataStorage,
+		UE::Editor::DataStorage::IUiProvider* DataStorageUi,
+		UE::Editor::DataStorage::RowHandle TargetRow,
+		UE::Editor::DataStorage::RowHandle WidgetRow,
+		const UE::Editor::DataStorage::FMetaDataView& Arguments) override;
+};
+
+
+USTRUCT()
+struct FEditorCrashTypeWidgetConstructor : public FSimpleWidgetConstructor
+{
+	GENERATED_BODY()
+
+public:
+	FEditorCrashTypeWidgetConstructor()
+		: Super(StaticStruct())
+	{
+	}
+
+	virtual TSharedPtr<SWidget> CreateWidget(
+		UE::Editor::DataStorage::ICoreProvider* DataStorage,
+		UE::Editor::DataStorage::IUiProvider* DataStorageUi,
+		UE::Editor::DataStorage::RowHandle TargetRow,
+		UE::Editor::DataStorage::RowHandle WidgetRow,
+		const UE::Editor::DataStorage::FMetaDataView& Arguments) override;
+};
+
+
+USTRUCT()
+struct FEditorCrashFileReportsWidgetConstructor : public FSimpleWidgetConstructor
+{
+	GENERATED_BODY()
+
+public:
+	FEditorCrashFileReportsWidgetConstructor()
+		: Super(StaticStruct())
+	{
+	}
+
+	virtual TSharedPtr<SWidget> CreateWidget(
+		UE::Editor::DataStorage::ICoreProvider* DataStorage,
+		UE::Editor::DataStorage::IUiProvider* DataStorageUi,
+		UE::Editor::DataStorage::RowHandle TargetRow,
+		UE::Editor::DataStorage::RowHandle WidgetRow,
+		const UE::Editor::DataStorage::FMetaDataView& Arguments) override;
+};
+
+
+USTRUCT()
+struct FEditorCrashCallStackWidgetConstructor : public FSimpleWidgetConstructor
+{
+	GENERATED_BODY()
+
+public:
+	FEditorCrashCallStackWidgetConstructor()
+		: Super(StaticStruct())
+	{
+	}
+
+	virtual TSharedPtr<SWidget> CreateWidget(
+		UE::Editor::DataStorage::ICoreProvider* DataStorage,
+		UE::Editor::DataStorage::IUiProvider* DataStorageUi,
+		UE::Editor::DataStorage::RowHandle TargetRow,
+		UE::Editor::DataStorage::RowHandle WidgetRow,
+		const UE::Editor::DataStorage::FMetaDataView& Arguments) override;
+};
+
+
+USTRUCT()
+struct FEditorCrashSourceContextWidgetConstructor : public FSimpleWidgetConstructor
+{
+	GENERATED_BODY()
+
+public:
+	FEditorCrashSourceContextWidgetConstructor()
+		: Super(StaticStruct())
+	{
+	}
+
+	virtual TSharedPtr<SWidget> CreateWidget(
+		UE::Editor::DataStorage::ICoreProvider* DataStorage,
+		UE::Editor::DataStorage::IUiProvider* DataStorageUi,
+		UE::Editor::DataStorage::RowHandle TargetRow,
+		UE::Editor::DataStorage::RowHandle WidgetRow,
+		const UE::Editor::DataStorage::FMetaDataView& Arguments) override;
+};
+
+
+USTRUCT()
+struct FEditorCrashUserActivityHintWidgetConstructor : public FSimpleWidgetConstructor
+{
+	GENERATED_BODY()
+
+public:
+	FEditorCrashUserActivityHintWidgetConstructor()
+		: Super(StaticStruct())
+	{
+	}
+
+	virtual TSharedPtr<SWidget> CreateWidget(
+		UE::Editor::DataStorage::ICoreProvider* DataStorage,
+		UE::Editor::DataStorage::IUiProvider* DataStorageUi,
+		UE::Editor::DataStorage::RowHandle TargetRow,
+		UE::Editor::DataStorage::RowHandle WidgetRow,
+		const UE::Editor::DataStorage::FMetaDataView& Arguments) override;
+};
+
+
+///////////////////////////////////
+// Header Widget Constructors
+///////////////////////////////////
+
+USTRUCT()
+struct FEditorCrashTimeHeaderWidgetConstructor final : public FTypedElementWidgetConstructor
+{
+	GENERATED_BODY()
+
+public:
+	FEditorCrashTimeHeaderWidgetConstructor()
+		: Super(StaticStruct())
+	{
+	}
+
+protected:
+	virtual TSharedPtr<SWidget> CreateWidget(const UE::Editor::DataStorage::FMetaDataView& Arguments) override;
+	virtual bool FinalizeWidget(
+		UE::Editor::DataStorage::ICoreProvider* DataStorage,
+		UE::Editor::DataStorage::IUiProvider* DataStorageUi,
+		UE::Editor::DataStorage::RowHandle Row,
+		const TSharedPtr<SWidget>& Widget) override;
+};
+
+
+USTRUCT()
+struct FEditorCrashTypeHeaderWidgetConstructor final : public FTypedElementWidgetConstructor
+{
+	GENERATED_BODY()
+
+public:
+	FEditorCrashTypeHeaderWidgetConstructor()
+		: Super(StaticStruct())
+	{
+	}
+
+protected:
+	virtual TSharedPtr<SWidget> CreateWidget(const UE::Editor::DataStorage::FMetaDataView& Arguments) override;
+	virtual bool FinalizeWidget(
+		UE::Editor::DataStorage::ICoreProvider* DataStorage,
+		UE::Editor::DataStorage::IUiProvider* DataStorageUi,
+		UE::Editor::DataStorage::RowHandle Row,
+		const TSharedPtr<SWidget>& Widget) override;
+};
\ No newline at end of file
diff --git a/Engine/Plugins/Experimental/EditorPerformance/Source/CrashDiagnostics/Private/TedsEditorCrashDataStorageFactory.cpp b/Engine/Plugins/Experimental/EditorPerformance/Source/CrashDiagnostics/Private/TedsEditorCrashDataStorageFactory.cpp
new file mode 100644
--- /dev/null
+++ b/Engine/Plugins/Experimental/EditorPerformance/Source/CrashDiagnostics/Private/TedsEditorCrashDataStorageFactory.cpp
@@ -0,0 +1,177 @@
+﻿// Copyright Epic Games, Inc. All Rights Reserved.
+
+#include "TedsEditorCrashDataStorageFactory.h"
+
+#include "CrashDiagnosticsModule.h"
+#include "TedsCrashColumnWidgets.h"
+#include "Brushes/SlateNoResource.h"
+#include "Columns/TedsCrashColumns.h"
+#include "DataStorage/Features.h"
+#include "Styling/CoreStyle.h"
+#include "Styling/SlateTypes.h"
+
+#define LOCTEXT_NAMESPACE "TedsEditorCrashDataStorageFactory"
+
+const UE::Editor::DataStorage::IUiProvider::FPurposeInfo& UTedsEditorCrashDataStorageFactory::GetTableRowPurpose()
+{
+	using namespace UE::Editor::DataStorage;
+	static IUiProvider::FPurposeInfo Purpose(
+		"CrashDiagnostics", "Cell", "TableRow",
+		IUiProvider::EPurposeType::UniqueByNameAndColumn,
+		LOCTEXT("CompactPurposeDescription", "Purpose to display widgets compactly in a table. (ex: show only one line for multiline texts)"),
+		GetGeneralPurposeID());
+	return Purpose;
+}
+
+UE::Editor::DataStorage::IUiProvider::FPurposeID UTedsEditorCrashDataStorageFactory::GetGeneralPurposeID()
+{
+	using namespace UE::Editor::DataStorage;
+	IUiProvider* DataStorageUi = GetMutableDataStorageFeature<IUiProvider>(UiFeatureName);
+	check(DataStorageUi)
+	return DataStorageUi->GetGeneralWidgetPurposeID();
+}
+
+UE::Editor::DataStorage::IUiProvider::FPurposeID UTedsEditorCrashDataStorageFactory::GetTableRowPurposeID()
+{
+	return GetTableRowPurpose().GeneratePurposeID();
+}
+
+UE::Editor::DataStorage::IUiProvider::FPurposeID UTedsEditorCrashDataStorageFactory::GetHeaderRowPurposeID()
+{
+	using namespace UE::Editor::DataStorage;
+	static IUiProvider::FPurposeInfo HeaderPurpose("General", "Header", NAME_None);
+	return HeaderPurpose.GeneratePurposeID();
+}
+
+FEditableTextBoxStyle* UTedsEditorCrashDataStorageFactory::GetReadOnlyTextBoxStyle()
+{
+	static FEditableTextBoxStyle Style = FEditableTextBoxStyle(FCoreStyle::Get().GetWidgetStyle<FEditableTextBoxStyle>("NormalEditableTextBox"))
+		.SetBackgroundImageReadOnly(FSlateNoResource());
+	return &Style;
+}
+
+void UTedsEditorCrashDataStorageFactory::RegisterTables(UE::Editor::DataStorage::ICoreProvider& DataStorage)
+{
+	DataStorage.RegisterTable<
+		FEditorCrashGUIDColumn,
+		FEditorCrashTimeColumn,
+		FEditorCrashTypeColumn,
+		FEditorCrashFileReportsColumn
+	>(UE::Editor::CrashDiagnostics::DataTableName);
+
+	DataStorage.RegisterTable(UE::Editor::CrashDiagnostics::GlobalTableName);
+}
+
+void UTedsEditorCrashDataStorageFactory::RegisterQueries(UE::Editor::DataStorage::ICoreProvider& DataStorage)
+{
+	using namespace UE::Editor::DataStorage;
+	using namespace UE::Editor::DataStorage::Queries;
+
+	HasCrashedLastSessionQueryHandle = DataStorage.RegisterQuery(
+		Count()
+		.Where()
+			.Any<FEditorCrashLastSessionTag>()
+		.Compile());
+
+	HasUnreadCrashesQueryHandle = DataStorage.RegisterQuery(
+		Count()
+		.Where()
+			.Any<FEditorCrashIsNewTag>()
+		.Compile());
+}
+
+void UTedsEditorCrashDataStorageFactory::RegisterWidgetPurposes(UE::Editor::DataStorage::IUiProvider& DataStorageUi) const
+{
+	DataStorageUi.RegisterWidgetPurpose(GetTableRowPurpose());
+}
+
+void UTedsEditorCrashDataStorageFactory::RegisterWidgetConstructors(
+	UE::Editor::DataStorage::ICoreProvider& DataStorage,
+	UE::Editor::DataStorage::IUiProvider& DataStorageUi) const
+{
+	using namespace UE::Editor::DataStorage::Queries;
+
+	///////////////////////////////////
+	// Widget Constructors
+	///////////////////////////////////
+
+	{
+		const IUiProvider::FPurposeID PurposeID = GetGeneralPurposeID();
+
+		DataStorageUi.RegisterWidgetFactory<FEditorCrashGUIDWidgetConstructor>(
+			DataStorageUi.FindPurpose(PurposeID),
+			TColumn<FEditorCrashGUIDColumn>()
+		);
+
+		DataStorageUi.RegisterWidgetFactory<FEditorCrashFileReportsWidgetConstructor>(
+			DataStorageUi.FindPurpose(PurposeID),
+			TColumn<FEditorCrashFileReportsColumn>()
+		);
+
+		DataStorageUi.RegisterWidgetFactory<FEditorCrashTimeWidgetConstructor>(
+			DataStorageUi.FindPurpose(PurposeID),
+			TColumn<FEditorCrashTimeColumn>()
+		);
+
+		DataStorageUi.RegisterWidgetFactory<FEditorCrashErrorMessageWidgetConstructor>(
+			DataStorageUi.FindPurpose(PurposeID),
+			TColumn<FEditorCrashErrorMessageColumn>()
+		);
+
+		DataStorageUi.RegisterWidgetFactory<FEditorCrashTypeWidgetConstructor>(
+			DataStorageUi.FindPurpose(PurposeID),
+			TColumn<FEditorCrashTypeColumn>()
+		);
+
+		DataStorageUi.RegisterWidgetFactory<FEditorCrashCallStackWidgetConstructor>(
+			DataStorageUi.FindPurpose(PurposeID),
+			TColumn<FEditorCrashCallStackColumn>()
+		);
+
+		DataStorageUi.RegisterWidgetFactory<FEditorCrashSourceContextWidgetConstructor>(
+			DataStorageUi.FindPurpose(PurposeID),
+			TColumn<FEditorCrashSourceContextColumn>()
+		);
+
+		DataStorageUi.RegisterWidgetFactory<FEditorCrashUserActivityHintWidgetConstructor>(
+			DataStorageUi.FindPurpose(PurposeID),
+			TColumn<FEditorCrashUserActivityHintColumn>()
+		);
+	}
+
+	///////////////////////////////////
+	// Compact Widget Constructors
+	///////////////////////////////////
+
+	{
+		const IUiProvider::FPurposeID PurposeID = GetTableRowPurposeID();
+
+		DataStorageUi.RegisterWidgetFactory<FEditorCrashTimeTableRowWidgetConstructor>(
+			DataStorageUi.FindPurpose(PurposeID),
+			TColumn<FEditorCrashTimeColumn>()
+		);
+
+		DataStorageUi.RegisterWidgetFactory<FEditorCrashErrorMessageCompactWidgetConstructor>(
+			DataStorageUi.FindPurpose(PurposeID),
+			TColumn<FEditorCrashErrorMessageColumn>()
+		);
+	}
+
+	///////////////////////////////////
+	// Header Widget Constructors
+	///////////////////////////////////
+
+	{
+		const IUiProvider::FPurposeID PurposeID = GetHeaderRowPurposeID();
+
+		DataStorageUi.RegisterWidgetFactory<FEditorCrashTimeHeaderWidgetConstructor>(
+			DataStorageUi.FindPurpose(PurposeID),
+			TColumn<FEditorCrashTimeColumn>());
+
+		DataStorageUi.RegisterWidgetFactory<FEditorCrashTypeHeaderWidgetConstructor>(
+			DataStorageUi.FindPurpose(PurposeID),
+			TColumn<FEditorCrashTypeColumn>());
+	}
+}
+
+#undef LOCTEXT_NAMESPACE
\ No newline at end of file
diff --git a/Engine/Plugins/Experimental/EditorPerformance/Source/CrashDiagnostics/Private/TedsEditorCrashDataStorageFactory.h b/Engine/Plugins/Experimental/EditorPerformance/Source/CrashDiagnostics/Private/TedsEditorCrashDataStorageFactory.h
new file mode 100644
--- /dev/null
+++ b/Engine/Plugins/Experimental/EditorPerformance/Source/CrashDiagnostics/Private/TedsEditorCrashDataStorageFactory.h
@@ -0,0 +1,45 @@
+﻿// Copyright Epic Games, Inc. All Rights Reserved.
+
+#pragma once
+
+#include "Elements/Interfaces/TypedElementDataStorageFactory.h"
+#include "Elements/Interfaces/TypedElementDataStorageUiInterface.h"
+
+#include "TedsEditorCrashDataStorageFactory.generated.h"
+
+struct FEditableTextBoxStyle;
+
+UCLASS()
+class UTedsEditorCrashDataStorageFactory : public UEditorDataStorageFactory
+{
+	GENERATED_BODY()
+
+public:
+	~UTedsEditorCrashDataStorageFactory() override = default;
+
+	static const UE::Editor::DataStorage::IUiProvider::FPurposeInfo& GetTableRowPurpose();
+	static UE::Editor::DataStorage::IUiProvider::FPurposeID GetGeneralPurposeID();
+	static UE::Editor::DataStorage::IUiProvider::FPurposeID GetTableRowPurposeID();
+	static UE::Editor::DataStorage::IUiProvider::FPurposeID GetHeaderRowPurposeID();
+
+	static FEditableTextBoxStyle* GetReadOnlyTextBoxStyle();
+
+	virtual void RegisterTables(UE::Editor::DataStorage::ICoreProvider& DataStorage) override;
+	virtual void RegisterQueries(UE::Editor::DataStorage::ICoreProvider& DataStorage) override;
+	virtual void RegisterWidgetPurposes(UE::Editor::DataStorage::IUiProvider& DataStorageUi) const override;
+	virtual void RegisterWidgetConstructors(
+		UE::Editor::DataStorage::ICoreProvider& DataStorage,
+		UE::Editor::DataStorage::IUiProvider& DataStorageUi) const override;
+
+	UE::Editor::DataStorage::QueryHandle GetHasCrashedLastSessionQueryHandle() const
+	{
+		return HasCrashedLastSessionQueryHandle;
+	}
+	UE::Editor::DataStorage::QueryHandle GetHasUnreadCrashesQueryHandle() const
+	{
+		return HasUnreadCrashesQueryHandle;
+	}
+private:
+	UE::Editor::DataStorage::QueryHandle HasCrashedLastSessionQueryHandle = UE::Editor::DataStorage::InvalidQueryHandle;
+	UE::Editor::DataStorage::QueryHandle HasUnreadCrashesQueryHandle = UE::Editor::DataStorage::InvalidQueryHandle;
+};
diff --git a/Engine/Plugins/Experimental/EditorPerformance/Source/CrashDiagnostics/Public/Columns/TedsCrashColumns.h b/Engine/Plugins/Experimental/EditorPerformance/Source/CrashDiagnostics/Public/Columns/TedsCrashColumns.h
new file mode 100644
--- /dev/null
+++ b/Engine/Plugins/Experimental/EditorPerformance/Source/CrashDiagnostics/Public/Columns/TedsCrashColumns.h
@@ -0,0 +1,127 @@
+﻿// Copyright Epic Games, Inc. All Rights Reserved.
+
+#pragma once
+
+#include "DataStorage/CommonTypes.h"
+#include "Misc/DateTime.h"
+#include "Misc/Paths.h"
+
+#include "TedsCrashColumns.generated.h"
+
+namespace UE::Editor::CrashDiagnostics
+{
+	inline static const FName MappingDomain = "CrashDiagnostics";
+	inline static const FName DataTableName = "CrashDiagnosticsDataTable";
+	inline static const FName GlobalTableName = "CrashDiagnosticsGlobalTable";
+	inline static const FName LastSessionCrashStateRowKey = "CrashDiagnosticsLastSessionCrashStateRowKey";
+}
+
+USTRUCT(DisplayName="File Reports")
+struct FEditorCrashFileReportsColumn : public FEditorDataStorageColumn
+{
+	GENERATED_BODY()
+
+	UPROPERTY(meta = (Searchable))
+	FString XmlFilePath;
+
+	UPROPERTY(meta = (Searchable))
+	FString LogFilePath;
+
+	UPROPERTY()
+	TArray<FString> ReportFilePaths;
+
+	FString GetCrashReportFolder() const
+	{
+		return FPaths::GetPath(XmlFilePath);
+	}
+};
+
+USTRUCT(DisplayName="Time of Crash")
+struct FEditorCrashTimeColumn : public FEditorDataStorageColumn
+{
+	GENERATED_BODY()
+
+	UPROPERTY()
+	FDateTime TimeOfCrash;
+
+	UPROPERTY(meta = (Searchable))
+	FText TimeOfCrashString;
+};
+
+USTRUCT(DisplayName="Crash ID")
+struct FEditorCrashGUIDColumn : public FEditorDataStorageColumn
+{
+	GENERATED_BODY()
+
+	UPROPERTY(meta = (Searchable, Sortable))
+	FString CrashGUID;
+};
+
+USTRUCT(DisplayName="Error Message")
+struct FEditorCrashErrorMessageColumn : public FEditorDataStorageColumn
+{
+	GENERATED_BODY()
+
+	UPROPERTY(meta = (Searchable))
+	FString ErrorMessage;
+};
+
+USTRUCT(DisplayName="Call Stack")
+struct FEditorCrashCallStackColumn : public FEditorDataStorageColumn
+{
+	GENERATED_BODY()
+
+	UPROPERTY(meta = (Searchable))
+	FString CallStack;
+};
+
+USTRUCT(DisplayName="Source Context")
+struct FEditorCrashSourceContextColumn : public FEditorDataStorageColumn
+{
+	GENERATED_BODY()
+
+	UPROPERTY(meta = (Searchable))
+	FString SourceContext;
+};
+
+USTRUCT(DisplayName="Crash Type")
+struct FEditorCrashTypeColumn : public FEditorDataStorageColumn
+{
+	GENERATED_BODY()
+
+	UPROPERTY(meta = (Searchable, Sortable))
+	FString CrashType;
+};
+
+USTRUCT(DisplayName="User Activity")
+struct FEditorCrashUserActivityHintColumn : public FEditorDataStorageColumn
+{
+	GENERATED_BODY()
+
+	UPROPERTY(meta = (Searchable, Sortable))
+	FString UserActivityHint;
+};
+
+USTRUCT(DisplayName="Is Ensure")
+struct FEditorCrashIsEnsureTag : public FEditorDataStorageTag
+{
+	GENERATED_BODY()
+};
+
+USTRUCT(DisplayName="Is Out Of Memory")
+struct FEditorCrashIsOOMTag : public FEditorDataStorageTag
+{
+	GENERATED_BODY()
+};
+
+USTRUCT(DisplayName="Is New Crash")
+struct FEditorCrashIsNewTag : public FEditorDataStorageTag
+{
+	GENERATED_BODY()
+};
+
+USTRUCT(DisplayName="Crashed Last Session")
+struct FEditorCrashLastSessionTag : public FEditorDataStorageTag
+{
+	GENERATED_BODY()
+};
diff --git a/Engine/Plugins/Experimental/EditorPerformance/Source/CrashDiagnostics/Public/CrashDiagnosticsModule.h b/Engine/Plugins/Experimental/EditorPerformance/Source/CrashDiagnostics/Public/CrashDiagnosticsModule.h
new file mode 100644
--- /dev/null
+++ b/Engine/Plugins/Experimental/EditorPerformance/Source/CrashDiagnostics/Public/CrashDiagnosticsModule.h
@@ -0,0 +1,75 @@
+﻿// Copyright Epic Games, Inc. All Rights Reserved.
+
+#pragma once
+
+#include "Misc/DateTime.h"
+#include "Modules/ModuleManager.h"
+
+#define UE_API CRASHDIAGNOSTICS_API
+
+struct FPrimaryCrashProperties;
+struct FSlateBrush;
+class SWidget;
+namespace UE::Editor::DataStorage
+{
+	class ICoreProvider;
+}
+
+namespace UE::Editor::CrashDiagnostics
+{
+
+class FCrashDiagnosticsModule : public IModuleInterface
+{
+public:
+	static FCrashDiagnosticsModule& GetChecked()
+	{
+		return FModuleManager::GetModuleChecked<FCrashDiagnosticsModule>("CrashDiagnostics");
+	}
+	static FCrashDiagnosticsModule* Get()
+	{
+		return FModuleManager::GetModulePtr<FCrashDiagnosticsModule>("CrashDiagnostics");
+	}
+
+	virtual void StartupModule() override;
+	virtual void ShutdownModule() override;
+
+	/**
+	 * Retrieve the crash data from the Saved directory
+	 */
+	static TArray<TSharedRef<FPrimaryCrashProperties>> RetrieveSavedCrashes();
+
+	/**
+	 * Add the crashes to the given DataStorage. See TedsCrashColumns.h
+	 */
+	void AddCrashesToDataStorage(UE::Editor::DataStorage::ICoreProvider& DataStorage);
+	void AddCrashesToDataStorage(UE::Editor::DataStorage::ICoreProvider& DataStorage, const TArray<TSharedRef<FPrimaryCrashProperties>>& Crashes);
+
+	/**
+	 * Add the crashes asynchronously to the default DataStorage.
+	 * This function will call RetrieveSavedCrashes on a background thread before calling AddCrashesToDataStorage on the GT.
+	 */
+	void AddCrashesToDataStorageAsync();
+
+	void LoadCrashDiagnosticsSettings();
+	void SaveCrashDiagnosticsSettings();
+
+	bool HasCrashedLastSession() const;
+	bool HasUnreadCrashes() const;
+
+	UE_API TSharedRef<SWidget> CreateCrashLogPanel();
+	/**
+	 * Function used for a GetDecoratedButtonDelegate that will add a badge to the first found SLayeredImage.
+	 * This badge will only show if there are unread crashes.
+	 */
+	static TSharedRef<SWidget> DecorateButtonWithUnreadBadge(TSharedRef<SWidget> InWidget);
+	const FSlateBrush* GetUnreadBadgeIcon() const;
+
+private:
+	void OnEditorInitialized(double TimeToInitializeEditor);
+
+	FDateTime PreviousSessionStartTime = FDateTime();
+};
+
+}
+
+#undef UE_API
diff --git a/Engine/Plugins/Experimental/EditorPerformance/Source/EditorPerformance/EditorPerformance.Build.cs b/Engine/Plugins/Experimental/EditorPerformance/Source/EditorPerformance/EditorPerformance.Build.cs
--- a/Engine/Plugins/Experimental/EditorPerformance/Source/EditorPerformance/EditorPerformance.Build.cs
+++ b/Engine/Plugins/Experimental/EditorPerformance/Source/EditorPerformance/EditorPerformance.Build.cs
@@ -12,6 +12,7 @@ public class EditorPerformance : ModuleRules
 			new string[] {
 				"Core",
 				"CoreUObject",
+				"CrashDiagnostics",
 				"Engine",
 				"Slate",
 				"ApplicationCore",
diff --git a/Engine/Plugins/Experimental/EditorPerformance/Source/EditorPerformance/Private/EditorPerformanceModule.cpp b/Engine/Plugins/Experimental/EditorPerformance/Source/EditorPerformance/Private/EditorPerformanceModule.cpp
--- a/Engine/Plugins/Experimental/EditorPerformance/Source/EditorPerformance/Private/EditorPerformanceModule.cpp
+++ b/Engine/Plugins/Experimental/EditorPerformance/Source/EditorPerformance/Private/EditorPerformanceModule.cpp
@@ -119,7 +119,7 @@ const FText UserDirFreeSpaceKPIText = LOCTEXT("UserDirFreeSpace", "User Dir Free
 
 // Thresholds
 TOptional<float> NoSetLimit;
-float LocalDDCEffciencyKPILimit = 85.f;
+float LocalDDCEfficiencyKPILimit = 85.f;
 float VirtualAssetEfficiencyKPILimit = 95.f;
 float MemoryPressureKPILimit = (float)FGenericPlatformMemoryStats::EMemoryPressureStatus::Warning;
 float DynamicInvalidationKPILimit = 1.0f;
@@ -265,7 +265,7 @@ void FEditorPerformanceModule::InitializeKPIs()
 	}
 	
 	{
-		LocalDDCEfficiencyKPI		= KPIRegistry.DeclareKPIValue(CacheCategoryName, CacheCategoryText, LocalDDCEfficiencyKPIName, LocalDDCEfficiencyKPIText, 100.0, LocalDDCEffciencyKPILimit, FKPIValue::GreaterThan, FKPIValue::Percent);
+		LocalDDCEfficiencyKPI		= KPIRegistry.DeclareKPIValue(CacheCategoryName, CacheCategoryText, LocalDDCEfficiencyKPIName, LocalDDCEfficiencyKPIText, 100.0, LocalDDCEfficiencyKPILimit, FKPIValue::GreaterThan, FKPIValue::Percent);
 		VirtualAssetEfficiencyKPI	= KPIRegistry.DeclareKPIValue(CacheCategoryName, CacheCategoryText, VirtualAssetEfficiencyKPIName, VirtualAssetEfficiencyKPIText, 100.0, VirtualAssetEfficiencyKPILimit, FKPIValue::GreaterThan, FKPIValue::Percent);
 	}
 	
@@ -726,7 +726,7 @@ void FEditorPerformanceModule::UpdateKPIsHighFrequency()
 					if (EditorPerformanceSettings->bEnableSnapshots)
 					{
 						// Create an Insights Snapshot
-						RecordInsightsSnaphshot(KPIValue);
+						RecordInsightsSnapshot(KPIValue);
 					}
 
 					if (EditorPerformanceSettings->bEnableTelemetry)
@@ -860,7 +860,7 @@ bool FEditorPerformanceModule::IsHotLocalCacheCase() const
 	return false;
 }
 
-bool FEditorPerformanceModule::RecordInsightsSnaphshot(const FKPIValue& KPIValue)
+bool FEditorPerformanceModule::RecordInsightsSnapshot(const FKPIValue& KPIValue)
 {
 	const uint32 MaxKPITraceCount = 10;
 
diff --git a/Engine/Plugins/Experimental/EditorPerformance/Source/EditorPerformance/Private/SEditorDiagnosticsTab.cpp b/Engine/Plugins/Experimental/EditorPerformance/Source/EditorPerformance/Private/SEditorDiagnosticsTab.cpp
--- a/Engine/Plugins/Experimental/EditorPerformance/Source/EditorPerformance/Private/SEditorDiagnosticsTab.cpp
+++ b/Engine/Plugins/Experimental/EditorPerformance/Source/EditorPerformance/Private/SEditorDiagnosticsTab.cpp
@@ -2,6 +2,7 @@
 
 #include "SEditorDiagnosticsTab.h"
 
+#include "CrashDiagnosticsModule.h"
 #include "SEditorPerformanceDialogs.h"
 #include "StallLogSubsystem.h"
 
@@ -31,6 +32,7 @@ void FEditorDiagnosticsCommands::RegisterCommands()
 {	
 	UI_COMMAND(SetShowActivityMonitor, "Activity", "Displays the activity monitor", EUserInterfaceActionType::ToggleButton, FInputChord());
 	UI_COMMAND(SetShowStalls, "Stalls", "Displays the stalls log", EUserInterfaceActionType::ToggleButton, FInputChord());
+	UI_COMMAND(SetShowCrashes, "Crashes", "Displays the crashes log", EUserInterfaceActionType::ToggleButton, FInputChord());
 }
 
 void SEditorDiagnosticsTab::Construct(const FArguments& InArgs)
@@ -98,7 +100,6 @@ TSharedRef<SWidget> SEditorDiagnosticsTab::CreateToolBar()
 	}
 
 	TWeakPtr<FUICommandInfo> ShowStallsCommand = FEditorDiagnosticsCommands::Get().SetShowStalls;
-
 	ToolBarCommandList->MapAction(FEditorDiagnosticsCommands::Get().SetShowStalls,
 		FExecuteAction::CreateRaw(this, &SEditorDiagnosticsTab::SetShowStalls),
 		FCanExecuteAction(),
@@ -107,6 +108,15 @@ TSharedRef<SWidget> SEditorDiagnosticsTab::CreateToolBar()
 				return SelectedCommand == ShowStallsCommand ? ECheckBoxState::Checked : ECheckBoxState::Unchecked;		
 			}));
 
+	TWeakPtr<FUICommandInfo> ShowCrashesCommand = FEditorDiagnosticsCommands::Get().SetShowCrashes;
+	ToolBarCommandList->MapAction(FEditorDiagnosticsCommands::Get().SetShowCrashes,
+		FExecuteAction::CreateRaw(this, &SEditorDiagnosticsTab::SetShowCrashes),
+		FCanExecuteAction(),
+		FGetActionCheckState::CreateLambda([this, ShowCrashesCommand] ()
+			{
+				return SelectedCommand == ShowCrashesCommand ? ECheckBoxState::Checked : ECheckBoxState::Unchecked;
+			}));
+
 	constexpr bool bForceSmallIcons = true; 
 	FVerticalToolBarBuilder ToolBarBuilder = FVerticalToolBarBuilder(ToolBarCommandList, FMultiBoxCustomization::None, TSharedPtr<FExtender>(), bForceSmallIcons);
 	ToolBarBuilder.SetLabelVisibility(EVisibility::Visible);
@@ -124,6 +134,7 @@ TSharedRef<SWidget> SEditorDiagnosticsTab::CreateToolBar()
 	}
 
 	ToolBarBuilder.AddToolBarButton(FEditorDiagnosticsCommands::Get().SetShowStalls);
+	ToolBarBuilder.AddToolBarButton(FEditorDiagnosticsCommands::Get().SetShowCrashes);
 
 	return ToolBarBuilder.MakeWidget();
 }
@@ -140,6 +151,13 @@ TSharedRef<SWidget> SEditorDiagnosticsTab::CreateStallLogPanel()
 	return StallLogSubsystem->CreateStallLogPanel();
 }
 
+TSharedRef<SWidget> SEditorDiagnosticsTab::CreateCrashesLogPanel()
+{
+	using namespace UE::Editor::CrashDiagnostics;
+	FCrashDiagnosticsModule& CrashDiagnostics = FCrashDiagnosticsModule::GetChecked();
+	return CrashDiagnostics.CreateCrashLogPanel();
+}
+
 void SEditorDiagnosticsTab::SetShowActivityMonitor()
 {
 	if (TSharedPtr<SBorder> BorderWidget = ContentBorder.Pin())
@@ -158,4 +176,13 @@ void SEditorDiagnosticsTab::SetShowStalls()
 	}
 }
 
+void SEditorDiagnosticsTab::SetShowCrashes()
+{
+	if (TSharedPtr<SBorder> BorderWidget = ContentBorder.Pin())
+	{
+		BorderWidget->SetContent(CreateCrashesLogPanel());
+		SelectedCommand = FEditorDiagnosticsCommands::Get().SetShowCrashes;
+	}
+}
+
 #undef LOCTEXT_NAMESPACE
diff --git a/Engine/Plugins/Experimental/EditorPerformance/Source/EditorPerformance/Private/SEditorDiagnosticsTab.h b/Engine/Plugins/Experimental/EditorPerformance/Source/EditorPerformance/Private/SEditorDiagnosticsTab.h
--- a/Engine/Plugins/Experimental/EditorPerformance/Source/EditorPerformance/Private/SEditorDiagnosticsTab.h
+++ b/Engine/Plugins/Experimental/EditorPerformance/Source/EditorPerformance/Private/SEditorDiagnosticsTab.h
@@ -2,6 +2,7 @@
 
 #pragma once
 
+#include "EditorPerformanceModule.h"
 #include "Framework/Commands/Commands.h"
 #include "Widgets/Docking/SDockTab.h"
 
@@ -21,6 +22,7 @@ public:
 public:
 	TSharedPtr<FUICommandInfo> SetShowActivityMonitor;
 	TSharedPtr<FUICommandInfo> SetShowStalls;
+	TSharedPtr<FUICommandInfo> SetShowCrashes;
 };
 
 class SEditorDiagnosticsTab : public SDockTab
@@ -33,9 +35,11 @@ protected:
 
 	TSharedRef<SWidget> CreateActivityMonitorPanel();
 	TSharedRef<SWidget> CreateStallLogPanel();
+	TSharedRef<SWidget> CreateCrashesLogPanel();
 
 	void SetShowActivityMonitor();
 	void SetShowStalls();
+	void SetShowCrashes();
 
 private:
 	TWeakPtr<SBorder> ContentBorder;
diff --git a/Engine/Plugins/Experimental/EditorPerformance/Source/EditorPerformance/Public/EditorPerformanceModule.h b/Engine/Plugins/Experimental/EditorPerformance/Source/EditorPerformance/Public/EditorPerformanceModule.h
--- a/Engine/Plugins/Experimental/EditorPerformance/Source/EditorPerformance/Public/EditorPerformanceModule.h
+++ b/Engine/Plugins/Experimental/EditorPerformance/Source/EditorPerformance/Public/EditorPerformanceModule.h
@@ -8,7 +8,6 @@
 #include "Modules/ModuleInterface.h"
 #include "Modules/ModuleManager.h"
 #include "Templates/SharedPointer.h"
-#include "Templates/UniquePtr.h"
 #include "KPIValue.h"
 
 #define UE_API EDITORPERFORMANCE_API
@@ -46,13 +45,13 @@ public:
 	 */
 	UE_API virtual void ShutdownModule() override;
 
-	UE_API TSharedRef<SWidget>	CreateStatusBarWidget();
+	UE_API TSharedRef<SWidget> CreateStatusBarWidget();
 	UE_API void ShowPerformanceReportTab();
 
 	UE_API const FKPIRegistry& GetKPIRegistry() const;
 	UE_API const FString& GetKPIProfileName() const;
 
-	UE_API bool RecordInsightsSnaphshot(const FKPIValue& Value);
+	UE_API bool RecordInsightsSnapshot(const FKPIValue& Value);
 	UE_API bool RecordTelemetryEvent(const FKPIValue& Value);
 	UE_API bool IsHotLocalCacheCase() const;
 	UE_API FEditorPerformanceModule::EEditorState GetEditorState() const;
diff --git a/Engine/Source/Editor/EditorStyle/Private/StarshipStyle.cpp b/Engine/Source/Editor/EditorStyle/Private/StarshipStyle.cpp
--- a/Engine/Source/Editor/EditorStyle/Private/StarshipStyle.cpp
+++ b/Engine/Source/Editor/EditorStyle/Private/StarshipStyle.cpp
@@ -1179,6 +1179,7 @@ void FStarshipEditorStyle::FStyle::SetupGeneralStyles()
 
 		Set("EditorDiagnostics.SetShowActivityMonitor", new IMAGE_BRUSH_SVG("Starship/Common/Scalability_16", Icon20x20));
 		Set("EditorDiagnostics.SetShowStalls", new IMAGE_BRUSH_SVG("Starship/Common/Stall", Icon20x20));
+		Set("EditorDiagnostics.SetShowCrashes", new IMAGE_BRUSH_SVG("Starship/Common/Debug", Icon20x20));
 
 		Set("EditorDiagnostics.StatusBarButton", GetWidgetStyle<FButtonStyle>("StatusBar.StatusBarButton"));
 
diff --git a/Engine/Source/Runtime/CrashReportCore/Private/CrashDescription.cpp b/Engine/Source/Runtime/CrashReportCore/Private/CrashDescription.cpp
--- a/Engine/Source/Runtime/CrashReportCore/Private/CrashDescription.cpp
+++ b/Engine/Source/Runtime/CrashReportCore/Private/CrashDescription.cpp
@@ -235,7 +235,10 @@ void FPrimaryCrashProperties::ReadXML( const FString& CrashContextFilepath, cons
 	{
 		XmlFile->LoadFile(XmlFilepath);
 	}
+}
 
+void FPrimaryCrashProperties::UpdateDefaults()
+{
 	if (XmlFile->IsValid())
 	{
 		TimeOfCrash = FDateTime::UtcNow().GetTicks();
@@ -422,7 +425,7 @@ namespace CrashDescriptionInternal
 	};
 }
 
-void FPrimaryCrashProperties::MakeCrashEventAttributes(TArray<FAnalyticsEventAttribute>& OutCrashAttributes)
+void FPrimaryCrashProperties::MakeCrashEventAttributes(TArray<FAnalyticsEventAttribute>& OutCrashAttributes) const
 {
 	OutCrashAttributes.Add(FAnalyticsEventAttribute(TEXT("bHasPrimaryData"), bHasPrimaryData));
 	OutCrashAttributes.Add(FAnalyticsEventAttribute(TEXT("CrashVersion"), (int32)CrashVersion));
@@ -529,19 +532,17 @@ void FPrimaryCrashProperties::Save()
 	FCrashContextReader
 -----------------------------------------------------------------------------*/
 
-FCrashContext::FCrashContext(const FString& CrashContextFilePath, const TCHAR* Buffer)
-{
-	ReadXML(CrashContextFilePath, Buffer);
-	SetupPrimaryCrashProperties();
-}
-
-FCrashContext::FCrashContext(const FString& CrashContextFilepath)
+FCrashContext::FCrashContext(const FString& CrashContextFilepath, const TCHAR* Buffer, bool bIsReadOnly)
 {
-	ReadXML(CrashContextFilepath);
+	ReadXML(CrashContextFilepath, Buffer);
+	if (!bIsReadOnly)
+	{
+		UpdateDefaults();
+	}
 	SetupPrimaryCrashProperties();
 }
 
-void FCrashContext::SetupPrimaryCrashProperties() 
+void FCrashContext::SetupPrimaryCrashProperties()
 {
 	const bool bIsValid = XmlFile->IsValid();
 	if (bIsValid)
@@ -616,10 +617,14 @@ void FCrashContext::SetupPrimaryCrashProperties()
 	FCrashDescription
 -----------------------------------------------------------------------------*/
 
-FCrashWERContext::FCrashWERContext( const FString& WERXMLFilepath )
+FCrashWERContext::FCrashWERContext( const FString& WERXMLFilepath, bool bIsReadOnly )
 	: FPrimaryCrashProperties()
 {
 	ReadXML( WERXMLFilepath );
+	if (!bIsReadOnly)
+	{
+		UpdateDefaults();
+	}
 	CrashGUID = FPaths::GetCleanFilename( FPaths::GetPath( WERXMLFilepath ) );
 
 	const bool bIsValid = XmlFile->IsValid();
diff --git a/Engine/Source/Runtime/CrashReportCore/Public/Android/AndroidErrorReport.h b/Engine/Source/Runtime/CrashReportCore/Public/Android/AndroidErrorReport.h
--- a/Engine/Source/Runtime/CrashReportCore/Public/Android/AndroidErrorReport.h
+++ b/Engine/Source/Runtime/CrashReportCore/Public/Android/AndroidErrorReport.h
@@ -26,7 +26,7 @@ public:
 	 * Discover all files in the crash report directory
 	 * @param Directory Full path to directory containing the report
 	 */
-	explicit FAndroidErrorReport(const FString& Directory);
+	CRASHREPORTCORE_API explicit FAndroidErrorReport(const FString& Directory);
 
 	/**
 	 * Load helper modules
@@ -40,7 +40,7 @@ public:
 
 	/**
 	 */
-	FText DiagnoseReport() const;
+	CRASHREPORTCORE_API FText DiagnoseReport() const;
 
 	/**
 	 * Do nothing - shouldn't be called on Android
diff --git a/Engine/Source/Runtime/CrashReportCore/Public/CrashDescription.h b/Engine/Source/Runtime/CrashReportCore/Public/CrashDescription.h
--- a/Engine/Source/Runtime/CrashReportCore/Public/CrashDescription.h
+++ b/Engine/Source/Runtime/CrashReportCore/Public/CrashDescription.h
@@ -128,13 +128,13 @@ public:
 	FCrashProperty& operator=(const int64 NewValue);
 
 	/** Getter for string, default. */
-	const FString& AsString() const;
+	CRASHREPORTCORE_API const FString& AsString() const;
 
 	/** Getter for bool. */
-	bool AsBool() const;
+	CRASHREPORTCORE_API bool AsBool() const;
 
 	/** Getter for int64. */
-	int64 AsInt64() const;
+	CRASHREPORTCORE_API int64 AsInt64() const;
 
 protected:
 	/** Owner of the property. */
@@ -481,7 +481,7 @@ struct FPrimaryCrashProperties
 
 protected:
 	/** Default constructor. */
-	FPrimaryCrashProperties();
+	CRASHREPORTCORE_API FPrimaryCrashProperties();
 
 	/** Destructor. */
 	~FPrimaryCrashProperties()
@@ -533,11 +533,17 @@ public:
 	/** Saves the data. */
 	void Save();
 
+	/** Gets the FilePath of the XML File. */
+	const FString& GetXmlFilepath() { return XmlFilepath; }
+
 protected:
 	/** Reads previously set XML file. Optionally give a buffer to parse from. Any changes will be stored to CrashContextFilePath. */
 	void ReadXML( const FString& CrashContextFilepath, const TCHAR* Buffer = nullptr );
 
-	/** Sets the CrasgGUID based on the report's path. */
+	/** Sets the TimeOfCrash and calls UpdateIDs if the XMLFile is valid */
+	void UpdateDefaults();
+
+	/** Sets the CrashGUID based on the report's path. */
 	void SetCrashGUID( const FString& Filepath );
 
 	/** Gets a crash property from the XML file. */
@@ -622,7 +628,7 @@ protected:
 	/** Encodes multi line property to be saved as single line. */
 	FString EncodeArrayStringAsXMLString( const TArray<FString>& ArrayString ) const;
 
-	void MakeCrashEventAttributes(TArray<FAnalyticsEventAttribute>& OutCrashAttributes);
+	CRASHREPORTCORE_API void MakeCrashEventAttributes(TArray<FAnalyticsEventAttribute>& OutCrashAttributes) const;
 
 	/** Reader for the xml file. */
 	FXmlFile* XmlFile;
@@ -645,10 +651,16 @@ protected:
 struct FCrashContext : public FPrimaryCrashProperties
 {
 	/** Initializes instance based on specified Crash Context filepath. */
-	explicit FCrashContext( const FString& CrashContextFilepath );
+	explicit FCrashContext( const FString& CrashContextFilepath, bool bIsReadOnly = false )
+		: FCrashContext(CrashContextFilepath, nullptr, bIsReadOnly)
+	{
+	}
 
 	/** Initializes an instance based on a in memory buffer. */
-	FCrashContext(const FString& CrashContextFilepath, const TCHAR* Buffer);
+	CRASHREPORTCORE_API FCrashContext(const FString& CrashContextFilepath, const TCHAR* Buffer, bool bIsReadOnly = false);
+
+	FCrashContext(const FCrashContext&) = delete;
+	FCrashContext& operator=(const FCrashContext&) = delete;
 
 private:
 
@@ -660,7 +672,10 @@ private:
 struct FCrashWERContext : public FPrimaryCrashProperties
 {
 	/** Initializes instance based on specified WER XML filepath. */
-	explicit FCrashWERContext( const FString& WERXMLFilepath );
+	CRASHREPORTCORE_API explicit FCrashWERContext( const FString& WERXMLFilepath, bool bIsReadOnly = false );
+
+	FCrashWERContext(const FCrashWERContext&) = delete;
+	FCrashWERContext& operator=(const FCrashWERContext&) = delete;
 
 	/** Initializes engine version from the separate components. */
 	void InitializeEngineVersion( const FString& BuildVersion, const FString& BranchName, uint32 BuiltFromCL );
diff --git a/Engine/Source/Runtime/CrashReportCore/Public/GenericErrorReport.h b/Engine/Source/Runtime/CrashReportCore/Public/GenericErrorReport.h
--- a/Engine/Source/Runtime/CrashReportCore/Public/GenericErrorReport.h
+++ b/Engine/Source/Runtime/CrashReportCore/Public/GenericErrorReport.h
@@ -30,7 +30,7 @@ public:
 	 * Discover all files in the crash report directory
 	 * @param Directory Full path to directory containing the report
 	 */
-	explicit FGenericErrorReport(const FString& Directory);
+	CRASHREPORTCORE_API explicit FGenericErrorReport(const FString& Directory);
 
 	/**
 	 * One-time initialisation: does nothing by default
@@ -63,13 +63,13 @@ public:
 	 * Provide full paths to all the report files
 	 * @return List of paths
 	 */
-	TArray<FString> GetFilesToUpload() const;
+	CRASHREPORTCORE_API TArray<FString> GetFilesToUpload() const;
 
 	/**
 	 * Provide full paths to all the report's optional attachment files
 	 * @return List of paths
 	 */
-	TArray<FString> GetOptionalAttachmentsToUpload() const;
+	CRASHREPORTCORE_API TArray<FString> GetOptionalAttachmentsToUpload() const;
 
 	/**
 	 * Load the WER XML file for this report
@@ -105,12 +105,12 @@ public:
 	/**
 	 * Get the name of the crashed app from the report
 	 */
-	FString FindCrashedAppName() const;
+	CRASHREPORTCORE_API FString FindCrashedAppName() const;
 
 	/**
 	 * Get the full path of the crashed app from the report
 	 */
-	FString FindCrashedAppPath() const;
+	CRASHREPORTCORE_API FString FindCrashedAppPath() const;
 
 	/**
 	 * Is there anything to upload?
@@ -125,7 +125,7 @@ public:
 	 * Look thought the list of report files to find one with the given extension
 	 * @return Whether a file with the extension was found
 	 */
-	bool FindFirstReportFileWithExtension(FString& OutFilename, const TCHAR* Extension) const;
+	CRASHREPORTCORE_API bool FindFirstReportFileWithExtension(FString& OutFilename, const TCHAR* Extension) const;
 
 	/**
 	 * Delete files from disk
diff --git a/Engine/Source/Runtime/CrashReportCore/Public/IOS/IOSErrorReport.h b/Engine/Source/Runtime/CrashReportCore/Public/IOS/IOSErrorReport.h
--- a/Engine/Source/Runtime/CrashReportCore/Public/IOS/IOSErrorReport.h
+++ b/Engine/Source/Runtime/CrashReportCore/Public/IOS/IOSErrorReport.h
@@ -31,18 +31,18 @@ public:
 	 * Discover all files in the crash report directory
 	 * @param Directory Full path to directory containing the report
 	 */
-	explicit FIOSErrorReport(const FString& Directory);
+	CRASHREPORTCORE_API explicit FIOSErrorReport(const FString& Directory);
 
 	/**
 	 * Parse the callstack from the Apple-style crash report log
 	 * @return UE crash diagnosis text
 	 */
-	FText DiagnoseReport() const;
+	CRASHREPORTCORE_API FText DiagnoseReport() const;
 
 	/**
 	 * Get the full path of the crashed app from the report (hides implementation in FGenericErrorReport)
 	 */
-	FString FindCrashedAppPath() const;
+	CRASHREPORTCORE_API FString FindCrashedAppPath() const;
 
 	/**
 	 * Look for the most recent Mac Error Report
diff --git a/Engine/Source/Runtime/CrashReportCore/Public/Mac/MacErrorReport.h b/Engine/Source/Runtime/CrashReportCore/Public/Mac/MacErrorReport.h
--- a/Engine/Source/Runtime/CrashReportCore/Public/Mac/MacErrorReport.h
+++ b/Engine/Source/Runtime/CrashReportCore/Public/Mac/MacErrorReport.h
@@ -31,18 +31,18 @@ public:
 	 * Discover all files in the crash report directory
 	 * @param Directory Full path to directory containing the report
 	 */
-	explicit FMacErrorReport(const FString& Directory);
+	CRASHREPORTCORE_API explicit FMacErrorReport(const FString& Directory);
 
 	/**
 	 * Parse the callstack from the Apple-style crash report log
 	 * @return UE crash diagnosis text
 	 */
-	FText DiagnoseReport() const;
+	CRASHREPORTCORE_API FText DiagnoseReport() const;
 
 	/**
 	 * Get the full path of the crashed app from the report (hides implementation in FGenericErrorReport)
 	 */
-	FString FindCrashedAppPath() const;
+	CRASHREPORTCORE_API FString FindCrashedAppPath() const;
 
 	/**
 	 * Look for the most recent Mac Error Report
diff --git a/Engine/Source/Runtime/CrashReportCore/Public/Windows/WindowsErrorReport.h b/Engine/Source/Runtime/CrashReportCore/Public/Windows/WindowsErrorReport.h
--- a/Engine/Source/Runtime/CrashReportCore/Public/Windows/WindowsErrorReport.h
+++ b/Engine/Source/Runtime/CrashReportCore/Public/Windows/WindowsErrorReport.h
@@ -31,18 +31,18 @@ public:
 	 * Discover all files in the crash report directory
 	 * @param Directory Full path to directory containing the report
 	 */
-	explicit FWindowsErrorReport(const FString& Directory);
+	CRASHREPORTCORE_API explicit FWindowsErrorReport(const FString& Directory);
 
 	/**
 	 * Provide the exception and a call-stack as plain text if possible
 	 * @note This can take quite a long time
 	 */
-	FText DiagnoseReport() const;
+	CRASHREPORTCORE_API FText DiagnoseReport() const;
 
 	/**
 	 * Get the full path of the crashed app from the report
 	 */
-	FString FindCrashedAppPath() const;
+	CRASHREPORTCORE_API FString FindCrashedAppPath() const;
 
 	/**
 	 * Look for the most recent Windows Error Reports
